#ue4: Blueprint Compiler Cpp Backend moved to its own module

diff --git a/Engine/Source/Developer/BlueprintCompilerCppBackend/BlueprintCompilerCppBackend.Build.cs b/Engine/Source/Developer/BlueprintCompilerCppBackend/BlueprintCompilerCppBackend.Build.cs
new file mode 100644
--- /dev/null
+++ b/Engine/Source/Developer/BlueprintCompilerCppBackend/BlueprintCompilerCppBackend.Build.cs
@@ -0,0 +1,33 @@
+// Copyright 1998-2015 Epic Games, Inc. All Rights Reserved.
+
+namespace UnrealBuildTool.Rules
+{
+	public class BlueprintCompilerCppBackend : ModuleRules
+	{
+        public BlueprintCompilerCppBackend(TargetInfo Target)
+		{
+			PublicIncludePaths.AddRange(
+				new string[] {
+					// ... add public include paths required here ...
+				}
+				);
+
+			PrivateIncludePaths.AddRange(
+				new string[] {
+					"Developer/BlueprintCompilerCppBackend/Private",
+				}
+				);
+
+			PublicDependencyModuleNames.AddRange(
+				new string[]
+				{
+					"Core",
+                    "CoreUObject",
+                    "Engine",
+					"KismetCompiler",
+                    "UnrealEd",
+				}
+				);
+		}
+	}
+}
\ No newline at end of file
diff --git a/Engine/Source/Developer/BlueprintCompilerCppBackend/Private/BlueprintCompilerCppBackend.cpp b/Engine/Source/Developer/BlueprintCompilerCppBackend/Private/BlueprintCompilerCppBackend.cpp
new file mode 100644
--- /dev/null
+++ b/Engine/Source/Developer/BlueprintCompilerCppBackend/Private/BlueprintCompilerCppBackend.cpp
@@ -0,0 +1,1054 @@
+// Copyright 1998-2015 Epic Games, Inc. All Rights Reserved.
+
+#include "BlueprintCompilerCppBackendModulePrivatePCH.h"
+#include "BlueprintCompilerCppBackend.h"
+#include "BlueprintCompilerCppBackendUtils.h"
+
+FString FBlueprintCompilerCppBackend::TermToText(const FBPTerminal* Term, const UProperty* CoerceProperty)
+{
+	const FString PSC_Self(TEXT("self"));
+	if (Term->bIsLiteral)
+	{
+		//@TODO: Must have a coercion type if it's a literal, because the symbol table isn't plumbed in here and the literals don't carry type information either, yay!
+		ensure(CoerceProperty);
+
+		if (CoerceProperty->IsA(UStrProperty::StaticClass()))
+		{
+			return FString::Printf(TEXT("TEXT(\"%s\")"), *(Term->Name));
+		}
+		else if (CoerceProperty->IsA(UTextProperty::StaticClass()))
+		{
+			return FString::Printf(TEXT("FText::FromString(TEXT(\"%s\"))"), *(Term->Name));
+		}
+		else if (CoerceProperty->IsA(UFloatProperty::StaticClass()))
+		{
+			float Value = FCString::Atof(*(Term->Name));
+			return FString::Printf(TEXT("%f"), Value);
+		}
+		else if (CoerceProperty->IsA(UIntProperty::StaticClass()))
+		{
+			int32 Value = FCString::Atoi(*(Term->Name));
+			return FString::Printf(TEXT("%d"), Value);
+		}
+		else if (auto ByteProperty = Cast<const UByteProperty>(CoerceProperty))
+		{
+			// The PinSubCategoryObject check is to allow enum literals communicate with byte properties as literals
+			if (ByteProperty->Enum != NULL ||
+				(NULL != Cast<UEnum>(Term->Type.PinSubCategoryObject.Get())))
+			{
+				return Term->Name;
+			}
+			else
+			{
+				uint8 Value = FCString::Atoi(*(Term->Name));
+				return FString::Printf(TEXT("%u"), Value);
+			}
+		}
+		else if (auto BoolProperty = Cast<const UBoolProperty>(CoerceProperty))
+		{
+			bool bValue = Term->Name.ToBool();
+			return bValue ? TEXT("true") : TEXT("false");
+		}
+		else if (auto NameProperty = Cast<const UNameProperty>(CoerceProperty))
+		{
+			FName LiteralName(*(Term->Name));
+			return FString::Printf(TEXT("FName(TEXT(\"%s\"))"), *(LiteralName.ToString()));
+		}
+		else if (auto StructProperty = Cast<const UStructProperty>(CoerceProperty))
+		{
+			if (StructProperty->Struct == VectorStruct)
+			{
+				FVector Vect = FVector::ZeroVector;
+				FDefaultValueHelper::ParseVector(Term->Name, /*out*/ Vect);
+
+				return FString::Printf(TEXT("FVector(%f,%f,%f)"), Vect.X, Vect.Y, Vect.Z);
+			}
+			else if (StructProperty->Struct == RotatorStruct)
+			{
+				FRotator Rot = FRotator::ZeroRotator;
+				FDefaultValueHelper::ParseRotator(Term->Name, /*out*/ Rot);
+
+				return FString::Printf(TEXT("FRotator(%f,%f,%f)"), Rot.Pitch, Rot.Yaw, Rot.Roll);
+			}
+			else if (StructProperty->Struct == TransformStruct)
+			{
+				FTransform Trans = FTransform::Identity;
+				Trans.InitFromString(Term->Name);
+
+				const FQuat Rot = Trans.GetRotation();
+				const FVector Translation = Trans.GetTranslation();
+				const FVector Scale = Trans.GetScale3D();
+
+				return FString::Printf(TEXT("FTransform( FQuat(%f,%f,%f,%f), FVector(%f,%f,%f), FVector(%f,%f,%f) )"),
+					Rot.X, Rot.Y, Rot.Z, Rot.W, Translation.X, Translation.Y, Translation.Z, Scale.X, Scale.Y, Scale.Z);
+			}
+			else if (StructProperty->Struct == LinearColorStruct)
+			{
+				FLinearColor LinearColor;
+				LinearColor.InitFromString(Term->Name);
+				return FString::Printf(TEXT("FLinearColor(%f,%f,%f,%f)"), LinearColor.R, LinearColor.G, LinearColor.B, LinearColor.A);
+			}
+			else
+			{
+				//@todo:  This needs to be more robust, since import text isn't really proper for struct construction.
+				return FString(TEXT("F")) + StructProperty->Struct->GetName() + Term->Name;
+			}
+		}
+		else if (auto ClassProperty = Cast<const UClassProperty>(CoerceProperty))
+		{
+			if (auto FoundClass = Cast<const UClass>(Term->ObjectLiteral))
+			{
+				return FString::Printf(TEXT("%s%s::StaticClass()"), FoundClass->GetPrefixCPP(), *FoundClass->GetName());
+			}
+			return FString(TEXT("NULL"));
+		}
+		else if (CoerceProperty->IsA(UDelegateProperty::StaticClass()))
+		{
+			//@TODO: K2 Delegate Support: Won't compile, there isn't an operator= that does what we want here, it should be a proper call to BindDynamic instead!
+			if (Term->Name == TEXT(""))
+			{
+				return TEXT("NULL");
+			}
+			else
+			{
+				return FString::Printf(TEXT("BindDynamic(this, &%s::%s)"), *CppClassName, *(Term->Name));
+			}
+		}
+		else if (CoerceProperty->IsA(UObjectPropertyBase::StaticClass()))
+		{
+			if (Term->Type.PinSubCategory == PSC_Self)
+			{
+				return TEXT("this");
+			}
+			else if (Term->ObjectLiteral)
+			{
+				if (auto LiteralClass = Cast<const UClass>(Term->ObjectLiteral))
+				{
+					return FString::Printf(TEXT("%s%s::StaticClass()"), LiteralClass->GetPrefixCPP(), *LiteralClass->GetName());
+				}
+				else
+				{
+					auto ObjectCoerceProperty = CastChecked<const UObjectPropertyBase>(CoerceProperty);
+					UClass* FoundClass = ObjectCoerceProperty ? ObjectCoerceProperty->PropertyClass : NULL;
+					FString ClassString = FoundClass ? (FString(FoundClass->GetPrefixCPP()) + FoundClass->GetName()) : TEXT("UObject");
+					return FString::Printf(TEXT("FindObject<%s>(ANY_PACKAGE, TEXT(\"%s\"))"), *ClassString, *(Term->ObjectLiteral->GetPathName()));
+				}
+			}
+			else
+			{
+				return FString(TEXT("NULL"));
+			}
+		}
+		else if (CoerceProperty->IsA(UInterfaceProperty::StaticClass()))
+		{
+			if (Term->Type.PinSubCategory == PSC_Self)
+			{
+				return TEXT("this");
+			}
+			else
+			{
+				ensureMsg(false, TEXT("It is not possible to express this interface property as a literal value!"));
+				return Term->Name;
+			}
+		}
+		else
+			// else if (CoerceProperty->IsA(UMulticastDelegateProperty::StaticClass()))
+			// Cannot assign a literal to a multicast delegate; it should be added instead of assigned
+		{
+			ensureMsg(false, TEXT("It is not possible to express this type as a literal value!"));
+			return Term->Name;
+		}
+	}
+	else
+	{
+		FString Prefix(TEXT(""));
+		if ((Term->Context != NULL) && (Term->Context->Name != PSC_Self))
+		{
+			Prefix = TermToText(Term->Context);
+
+			if (Term->Context->bIsStructContext)
+			{
+				Prefix += TEXT(".");
+			}
+			else
+			{
+				Prefix += TEXT("->");
+			}
+		}
+
+		Prefix += Term->Name;
+		return Prefix;
+	}
+}
+
+FString FBlueprintCompilerCppBackend::LatentFunctionInfoTermToText(FBPTerminal* Term, FBlueprintCompiledStatement* TargetLabel)
+{
+	check(LatentInfoStruct);
+
+	// Find the term name we need to fixup
+	FString FixupTermName;
+	for (UProperty* Prop = LatentInfoStruct->PropertyLink; Prop; Prop = Prop->PropertyLinkNext)
+	{
+		static const FName NeedsLatentFixup(TEXT("NeedsLatentFixup"));
+		if (Prop->GetBoolMetaData(NeedsLatentFixup))
+		{
+			FixupTermName = Prop->GetName();
+			break;
+		}
+	}
+
+	check(!FixupTermName.IsEmpty());
+
+	FString StructValues = Term->Name;
+
+	// Index 0 is always the ubergraph
+	const int32 TargetStateIndex = StateMapPerFunction[0].StatementToStateIndex(TargetLabel);
+	const int32 LinkageTermStartIdx = StructValues.Find(FixupTermName);
+	check(LinkageTermStartIdx != INDEX_NONE);
+	StructValues = StructValues.Replace(TEXT("-1"), *FString::FromInt(TargetStateIndex));
+
+	return FString(TEXT("F")) + LatentInfoStruct->GetName() + StructValues;
+}
+
+void FBlueprintCompilerCppBackend::EmitClassProperties(FStringOutputDevice& Target, UClass* SourceClass)
+{
+	// Emit class variables
+	for (TFieldIterator<UProperty> It(SourceClass, EFieldIteratorFlags::ExcludeSuper); It; ++It)
+	{
+		UProperty* Property = *It;
+		check(Property);
+
+		Emit(Header, TEXT("\n\tUPROPERTY("));
+		{
+			TArray<FString> Tags = FEmitHelper::ProperyFlagsToTags(Property->PropertyFlags);
+			Tags.Emplace(FEmitHelper::HandleRepNotifyFunc(Property));
+			Tags.Emplace(FEmitHelper::HandleMetaData(Property, false));
+			Tags.Remove(FString());
+
+			FString AllTags;
+			FEmitHelper::ArrayToString(Tags, AllTags, TEXT(", "));
+			Emit(Header, *AllTags);
+		}
+		Emit(Header, TEXT(")\n"));
+		Emit(Header, TEXT("\t"));
+		Property->ExportCppDeclaration(Target, EExportedDeclaration::Member, NULL, EPropertyExportCPPFlags::CPPF_CustomTypeName);
+		Emit(Header, TEXT(";\n"));
+	}
+}
+
+void FBlueprintCompilerCppBackend::GenerateCodeFromClass(UClass* SourceClass, FString NewClassName, TIndirectArray<FKismetFunctionContext>& Functions, bool bGenerateStubsOnly)
+{
+	auto CleanCppClassName = NewClassName.IsEmpty() ? SourceClass->GetName() : NewClassName;
+	CppClassName = FString(SourceClass->GetPrefixCPP()) + CleanCppClassName;
+
+	UClass* SuperClass = SourceClass->GetSuperClass();
+
+	Emit(Header, TEXT("#pragma once\n\n"));
+	Emit(Header, *FString::Printf(TEXT("#include \"%s.generated.h\"\n\n"), *CleanCppClassName));
+
+	// MC DELEGATE DECLARATION
+	{
+		auto DelegateDeclarations = FEmitHelper::EmitMulticastDelegateDeclarations(SourceClass);
+		FString AllDeclarations;
+		FEmitHelper::ArrayToString(DelegateDeclarations, AllDeclarations, TEXT(";\n"));
+		if (DelegateDeclarations.Num())
+		{
+			Emit(Header, *AllDeclarations);
+			Emit(Header, TEXT(";\n"));
+		}
+	}
+
+	// GATHER ALL SC DELEGATES
+	{
+		TArray<UDelegateProperty*> Delegates;
+		for (TFieldIterator<UDelegateProperty> It(SourceClass, EFieldIteratorFlags::ExcludeSuper); It; ++It)
+		{
+			Delegates.Add(*It);
+		}
+
+		for (auto& FuncContext : Functions)
+		{
+			for (TFieldIterator<UDelegateProperty> It(FuncContext.Function, EFieldIteratorFlags::ExcludeSuper); It; ++It)
+			{
+				Delegates.Add(*It);
+			}
+		}
+
+		auto DelegateDeclarations = FEmitHelper::EmitSinglecastDelegateDeclarations(Delegates);
+		FString AllDeclarations;
+		FEmitHelper::ArrayToString(DelegateDeclarations, AllDeclarations, TEXT(";\n"));
+		if (DelegateDeclarations.Num())
+		{
+			Emit(Header, *AllDeclarations);
+			Emit(Header, TEXT(";\n"));
+		}
+	}
+
+	Emit(Header, TEXT("UCLASS(Blueprintable)\n"));
+
+	Emit(Header,
+		*FString::Printf(TEXT("class %s : public %s%s\n"), *CppClassName, SuperClass->GetPrefixCPP(), *SuperClass->GetName()));
+	Emit(Header,
+		TEXT("{\n")
+		TEXT("public:\n"));
+	Emit(Header, *FString::Printf(TEXT("\tGENERATED_UCLASS_BODY()\n")));
+
+	EmitClassProperties(Header, SourceClass);
+
+	// Create the state map
+	for (int32 i = 0; i < Functions.Num(); ++i)
+	{
+		StateMapPerFunction.Add(FFunctionLabelInfo());
+		FunctionIndexMap.Add(&Functions[i], i);
+	}
+
+	// Emit function declarations and definitions (writes to header and body simultaneously)
+	if (Functions.Num() > 0)
+	{
+		Emit(Header, TEXT("\n"));
+	}
+
+	TArray<FString> PersistentHeaders;
+	PersistentHeaders.Add(FString(FApp::GetGameName()) + TEXT(".h"));
+	PersistentHeaders.Add(CleanCppClassName + TEXT(".h"));
+	PersistentHeaders.Add(TEXT("GeneratedCodeHelpers.h"));
+	Emit(Body, *FEmitHelper::GatherHeadersToInclude(SourceClass, PersistentHeaders));
+
+	//constructor
+	Emit(Body, *FString::Printf(
+		TEXT("%s::%s(const FObjectInitializer& ObjectInitializer) : Super(ObjectInitializer) {}\n\n"),
+		*CppClassName, *CppClassName));
+
+	for (int32 i = 0; i < Functions.Num(); ++i)
+	{
+		if (Functions[i].IsValid())
+		{
+			ConstructFunction(Functions[i], bGenerateStubsOnly);
+		}
+	}
+
+	Emit(Header, TEXT("};\n\n"));
+
+	Emit(Body, *FEmitHelper::EmitLifetimeReplicatedPropsImpl(SourceClass, CppClassName, TEXT("")));
+}
+
+void FBlueprintCompilerCppBackend::EmitCallDelegateStatment(FKismetFunctionContext& FunctionContext, FBlueprintCompiledStatement& Statement)
+{
+	check(Statement.FunctionContext && Statement.FunctionContext->AssociatedVarProperty);
+	FSafeContextScopedEmmitter SafeContextScope(Body, Statement.FunctionContext->Context, *this, TEXT("\t\t\t"));
+
+	Emit(Body, TEXT("\t\t\t"));
+	Emit(Body, *SafeContextScope.GetAdditionalIndent());
+	Emit(Body, *FString::Printf(TEXT("%s.Broadcast("), *TermToText(Statement.FunctionContext)));
+	int32 NumParams = 0;
+	for (TFieldIterator<UProperty> PropIt(Statement.FunctionToCall); PropIt && (PropIt->PropertyFlags & CPF_Parm); ++PropIt)
+	{
+		UProperty* FuncParamProperty = *PropIt;
+
+		if (!FuncParamProperty->HasAnyPropertyFlags(CPF_ReturnParm))
+		{
+			if (NumParams > 0)
+			{
+				Emit(Body, TEXT(", "));
+			}
+
+			FString VarName;
+
+			FBPTerminal* Term = Statement.RHS[NumParams];
+			ensure(Term != NULL);
+
+			// See if this is a hidden array param term, which needs to be fixed up with the final generated UArrayProperty
+			if (FBPTerminal** ArrayParmTerm = Statement.ArrayCoersionTermMap.Find(Term))
+			{
+				Term->ObjectLiteral = (*ArrayParmTerm)->AssociatedVarProperty;
+			}
+
+			if ((Statement.TargetLabel != NULL) && (Statement.UbergraphCallIndex == NumParams))
+			{
+				// The target label will only ever be set on a call function when calling into the Ubergraph or
+				// on a latent function that will later call into the ubergraph, either of which requires a patchup
+				UStructProperty* StructProp = Cast<UStructProperty>(FuncParamProperty);
+				if (StructProp && StructProp->Struct == LatentInfoStruct)
+				{
+					// Latent function info case
+					VarName = LatentFunctionInfoTermToText(Term, Statement.TargetLabel);
+				}
+				else
+				{
+					// Ubergraph entry point case
+					VarName = FString::FromInt(StateMapPerFunction[0].StatementToStateIndex(Statement.TargetLabel));
+				}
+			}
+			else
+			{
+				// Emit a normal parameter term
+				VarName = TermToText(Term, FuncParamProperty);
+			}
+
+			if (FuncParamProperty->HasAnyPropertyFlags(CPF_OutParm))
+			{
+				Emit(Body, TEXT("/*out*/ "));
+			}
+			Emit(Body, *VarName);
+
+			NumParams++;
+		}
+	}
+	Emit(Body, TEXT(");\n"));
+}
+
+void FBlueprintCompilerCppBackend::EmitCallStatment(FKismetFunctionContext& FunctionContext, FBlueprintCompiledStatement& Statement)
+{
+	const bool bCallOnDifferentObject = (Statement.FunctionContext != NULL) && (Statement.FunctionContext->Name != TEXT("self"));
+	const bool bStaticCall = Statement.FunctionToCall->HasAnyFunctionFlags(FUNC_Static);
+
+	const bool bUseSafeContext = bCallOnDifferentObject && !bStaticCall;
+	FSafeContextScopedEmmitter SafeContextScope(Body, bUseSafeContext ? Statement.FunctionContext : NULL, *this, TEXT("\t\t\t"));
+
+	Emit(Body, TEXT("\t\t\t"));
+	Emit(Body, *SafeContextScope.GetAdditionalIndent());
+
+	// Handle the return value of the function being called
+	UProperty* FuncToCallReturnProperty = Statement.FunctionToCall->GetReturnProperty();
+	if (FuncToCallReturnProperty != NULL)
+	{
+		FBPTerminal* Term = Statement.LHS;
+		ensure(Term != NULL);
+
+		FString ReturnVarName = TermToText(Term);
+		Emit(Body, *FString::Printf(TEXT("%s = "), *ReturnVarName));
+	}
+
+	// Emit object to call the method on
+	if (bCallOnDifferentObject) //@TODO: Badness, could be a self reference wired to another instance!
+	{
+		if (bStaticCall)
+		{
+			const bool bIsCustomThunk = Statement.FunctionToCall->HasMetaData(TEXT("CustomStructureParam")) || Statement.FunctionToCall->HasMetaData(TEXT("ArrayParm"));
+			FString FullFunctionNamePrefix = bIsCustomThunk
+				? TEXT("FCustomThunkTemplates::")
+				: FString::Printf(TEXT("U%s::"), *Statement.FunctionToCall->GetOuter()->GetName());
+			Emit(Body, *FullFunctionNamePrefix); //@TODO: Assuming U prefix but could be A
+		}
+		else
+		{
+			Emit(Body, *FString::Printf(TEXT("%s->"), *TermToText(Statement.FunctionContext, (UProperty*)(GetDefault<UObjectProperty>()))));
+		}
+	}
+
+	// Emit method name
+	FString FunctionNameToCall;
+	Statement.FunctionToCall->GetName(FunctionNameToCall);
+	if (Statement.bIsParentContext)
+	{
+		FunctionNameToCall = TEXT("Super::") + FunctionNameToCall;
+	}
+	Emit(Body, *FString::Printf(TEXT("%s"), *FunctionNameToCall));
+
+	// Emit method parameter list
+	Emit(Body, TEXT("("));
+	{
+		int32 NumParams = 0;
+
+		for (TFieldIterator<UProperty> PropIt(Statement.FunctionToCall); PropIt && (PropIt->PropertyFlags & CPF_Parm); ++PropIt)
+		{
+			UProperty* FuncParamProperty = *PropIt;
+
+			if (!FuncParamProperty->HasAnyPropertyFlags(CPF_ReturnParm))
+			{
+				if (NumParams > 0)
+				{
+					Emit(Body, TEXT(", "));
+				}
+
+				FString VarName;
+
+				FBPTerminal* Term = Statement.RHS[NumParams];
+				ensure(Term != NULL);
+
+				// See if this is a hidden array param term, which needs to be fixed up with the final generated UArrayProperty
+				if (FBPTerminal** ArrayParmTerm = Statement.ArrayCoersionTermMap.Find(Term))
+				{
+					Term->ObjectLiteral = (*ArrayParmTerm)->AssociatedVarProperty;
+				}
+
+				if ((Statement.TargetLabel != NULL) && (Statement.UbergraphCallIndex == NumParams))
+				{
+					// The target label will only ever be set on a call function when calling into the Ubergraph or
+					// on a latent function that will later call into the ubergraph, either of which requires a patchup
+					UStructProperty* StructProp = Cast<UStructProperty>(FuncParamProperty);
+					if (StructProp && StructProp->Struct == LatentInfoStruct)
+					{
+						// Latent function info case
+						VarName = LatentFunctionInfoTermToText(Term, Statement.TargetLabel);
+					}
+					else
+					{
+						// Ubergraph entry point case
+						VarName = FString::FromInt(StateMapPerFunction[0].StatementToStateIndex(Statement.TargetLabel));
+					}
+				}
+				else
+				{
+					// Emit a normal parameter term
+					VarName = TermToText(Term, FuncParamProperty);
+				}
+
+				if (FuncParamProperty->HasAnyPropertyFlags(CPF_OutParm))
+				{
+					Emit(Body, TEXT("/*out*/ "));
+				}
+				Emit(Body, *VarName);
+
+				NumParams++;
+			}
+		}
+	}
+	Emit(Body, TEXT(");\n"));
+}
+
+void FBlueprintCompilerCppBackend::EmitAssignmentStatment(FKismetFunctionContext& FunctionContext, FBlueprintCompiledStatement& Statement)
+{
+	FString DestinationExpression = TermToText(Statement.LHS);
+	FString SourceExpression = TermToText(Statement.RHS[0], Statement.LHS->AssociatedVarProperty);
+
+	FSafeContextScopedEmmitter SafeContextScope(Body, Statement.LHS->Context, *this, TEXT("\t\t\t"));
+	// Emit the assignment statement
+	Emit(Body, TEXT("\t\t\t"));
+	Emit(Body, *SafeContextScope.GetAdditionalIndent());
+	Emit(Body, *FString::Printf(TEXT("%s = %s;\n"), *DestinationExpression, *SourceExpression));
+}
+
+void FBlueprintCompilerCppBackend::EmitCastObjToInterfaceStatement(FKismetFunctionContext& FunctionContext, FBlueprintCompiledStatement& Statement)
+{
+	FString InterfaceClass = TermToText(Statement.RHS[0], (UProperty*)(GetDefault<UClassProperty>()));
+	FString ObjectValue = TermToText(Statement.RHS[1], (UProperty*)(GetDefault<UObjectProperty>()));
+	FString InterfaceValue = TermToText(Statement.LHS, (UProperty*)(GetDefault<UObjectProperty>()));
+
+	Emit(Body, *FString::Printf(TEXT("\t\t\tif ( IsValid(%s) && %s->GetClass()->ImplementsInterface(%s) )\n"), *ObjectValue, *ObjectValue, *InterfaceClass));
+	Emit(Body, *FString::Printf(TEXT("\t\t\t{\n")));
+	Emit(Body, *FString::Printf(TEXT("\t\t\t\t%s.SetObject(%s);\n"), *InterfaceValue, *ObjectValue));
+	Emit(Body, *FString::Printf(TEXT("\t\t\t\tvoid* IAddress = %s->GetInterfaceAddress(%s);\n"), *ObjectValue, *InterfaceClass));
+	Emit(Body, *FString::Printf(TEXT("\t\t\t\t%s.SetInterface(IAddress);\n"), *InterfaceValue));
+	Emit(Body, *FString::Printf(TEXT("\t\t\t}\n")));
+	Emit(Body, *FString::Printf(TEXT("\t\t\telse\n")));
+	Emit(Body, *FString::Printf(TEXT("\t\t\t{\n")));
+	Emit(Body, *FString::Printf(TEXT("\t\t\t\t%s.SetObject(NULL);\n"), *InterfaceValue));
+	Emit(Body, *FString::Printf(TEXT("\t\t\t}\n")));
+}
+
+void FBlueprintCompilerCppBackend::EmitCastBetweenInterfacesStatement(FKismetFunctionContext& FunctionContext, FBlueprintCompiledStatement& Statement)
+{
+	FString ClassToCastTo = TermToText(Statement.RHS[0], (UProperty*)(GetDefault<UClassProperty>()));
+	FString InputInterface = TermToText(Statement.RHS[1], (UProperty*)(GetDefault<UInterfaceProperty>()));
+	FString ResultInterface = TermToText(Statement.LHS, (UProperty*)(GetDefault<UInterfaceProperty>()));
+
+	FString InputObject = FString::Printf(TEXT("%s.GetObjectRef()"), *InputInterface);
+
+	Emit(Body, *FString::Printf(TEXT("\t\t\tif ( %s && %s->GetClass()->IsChildOf(%s) )\n"), *InputObject, *InputObject, *ClassToCastTo));
+	Emit(Body, *FString::Printf(TEXT("\t\t\t{\n")));
+	Emit(Body, *FString::Printf(TEXT("\t\t\t\t%s.SetObject(%s);\n"), *ResultInterface, *InputObject));
+	Emit(Body, *FString::Printf(TEXT("\t\t\t\tvoid* IAddress = %s->GetInterfaceAddress(%s);\n"), *InputObject, *ClassToCastTo));
+	Emit(Body, *FString::Printf(TEXT("\t\t\t\t%s.SetInterface(IAddress);\n"), *ResultInterface));
+	Emit(Body, *FString::Printf(TEXT("\t\t\t}\n")));
+	Emit(Body, *FString::Printf(TEXT("\t\t\telse\n")));
+	Emit(Body, *FString::Printf(TEXT("\t\t\t{\n")));
+	Emit(Body, *FString::Printf(TEXT("\t\t\t\t%s.SetObject(NULL);\n"), *ResultInterface));
+	Emit(Body, *FString::Printf(TEXT("\t\t\t}\n")));
+}
+
+void FBlueprintCompilerCppBackend::EmitCastInterfaceToObjStatement(FKismetFunctionContext& FunctionContext, FBlueprintCompiledStatement& Statement)
+{
+	FString ClassToCastTo = TermToText(Statement.RHS[0], (UProperty*)(GetDefault<UClassProperty>()));
+	FString InputInterface = TermToText(Statement.RHS[1], (UProperty*)(GetDefault<UInterfaceProperty>()));
+	FString ResultObject = TermToText(Statement.LHS, (UProperty*)(GetDefault<UInterfaceProperty>()));
+
+	FString InputObject = FString::Printf(TEXT("%s.GetObjectRef()"), *InputInterface);
+
+	Emit(Body, *FString::Printf(TEXT("\t\t\t%s = Cast<%s>(%s);\n"),
+		*ResultObject, *ClassToCastTo, *InputObject));
+}
+
+void FBlueprintCompilerCppBackend::EmitDynamicCastStatement(FKismetFunctionContext& FunctionContext, FBlueprintCompiledStatement& Statement)
+{
+	FString TargetClass = TermToText(Statement.RHS[0], (UProperty*)(GetDefault<UClassProperty>()));
+	FString ObjectValue = TermToText(Statement.RHS[1], (UProperty*)(GetDefault<UObjectProperty>()));
+	FString CastedValue = TermToText(Statement.LHS, (UProperty*)(GetDefault<UObjectProperty>()));
+
+	Emit(Body, *FString::Printf(TEXT("\t\t\t%s = Cast<%s>(%s);\n"),
+		*CastedValue, *TargetClass, *ObjectValue));
+}
+
+void FBlueprintCompilerCppBackend::EmitMetaCastStatement(FKismetFunctionContext& FunctionContext, FBlueprintCompiledStatement& Statement)
+{
+	FString DesiredClass = TermToText(Statement.RHS[0], (UProperty*)(GetDefault<UClassProperty>()));
+	FString SourceClass = TermToText(Statement.RHS[1], (UProperty*)(GetDefault<UClassProperty>()));
+	FString Destination = TermToText(Statement.LHS, (UProperty*)(GetDefault<UClassProperty>()));
+
+	Emit(Body, *FString::Printf(TEXT("\t\t\t%s = DynamicMetaCast(%s, %s);\n"),
+		*Destination, *DesiredClass, *SourceClass));
+}
+
+void FBlueprintCompilerCppBackend::EmitObjectToBoolStatement(FKismetFunctionContext& FunctionContext, FBlueprintCompiledStatement& Statement)
+{
+	FString ObjectTarget = TermToText(Statement.RHS[0]);
+	FString DestinationExpression = TermToText(Statement.LHS);
+
+	Emit(Body, *FString::Printf(TEXT("\t\t\t%s = (%s != NULL);\n"), *DestinationExpression, *ObjectTarget));
+}
+
+void FBlueprintCompilerCppBackend::EmitAddMulticastDelegateStatement(FKismetFunctionContext& FunctionContext, FBlueprintCompiledStatement& Statement)
+{
+	check(Statement.LHS && Statement.LHS->AssociatedVarProperty);
+	FSafeContextScopedEmmitter SafeContextScope(Body, Statement.LHS->Context, *this, TEXT("\t\t\t"));
+
+	Emit(Body, TEXT("\t\t\t"));
+	Emit(Body, *SafeContextScope.GetAdditionalIndent());
+
+	const FString Delegate = TermToText(Statement.LHS);
+	const FString DelegateToAdd = TermToText(Statement.RHS[0]);
+
+	Emit(Body, *FString::Printf(TEXT("%s.Add(%s);\n"), *Delegate, *DelegateToAdd));
+}
+
+void FBlueprintCompilerCppBackend::EmitRemoveMulticastDelegateStatement(FKismetFunctionContext& FunctionContext, FBlueprintCompiledStatement& Statement)
+{
+	check(Statement.LHS && Statement.LHS->AssociatedVarProperty);
+	FSafeContextScopedEmmitter SafeContextScope(Body, Statement.LHS->Context, *this, TEXT("\t\t\t"));
+
+	Emit(Body, TEXT("\t\t\t"));
+	Emit(Body, *SafeContextScope.GetAdditionalIndent());
+
+	const FString Delegate = TermToText(Statement.LHS);
+	const FString DelegateToAdd = TermToText(Statement.RHS[0]);
+
+	Emit(Body, *FString::Printf(TEXT("%s.Remove(%s);\n"), *Delegate, *DelegateToAdd));
+}
+
+void FBlueprintCompilerCppBackend::EmitBindDelegateStatement(FKismetFunctionContext& FunctionContext, FBlueprintCompiledStatement& Statement)
+{
+	check(2 == Statement.RHS.Num());
+	check(Statement.LHS);
+	FSafeContextScopedEmmitter SafeContextScope(Body, Statement.LHS->Context, *this, TEXT("\t\t\t"));
+
+	Emit(Body, TEXT("\t\t\t"));
+	Emit(Body, *SafeContextScope.GetAdditionalIndent());
+
+
+	const FString Delegate = TermToText(Statement.LHS);
+	const FString NameTerm = TermToText(Statement.RHS[0], GetDefault<UNameProperty>());
+	const FString ObjectTerm = TermToText(Statement.RHS[1], GetDefault<UObjectProperty>());
+
+	Emit(Body, *FString::Printf(TEXT("%s.BindUFunction(%s,%s);\n"), *Delegate, *ObjectTerm, *NameTerm));
+}
+
+void FBlueprintCompilerCppBackend::EmitClearMulticastDelegateStatement(FKismetFunctionContext& FunctionContext, FBlueprintCompiledStatement& Statement)
+{
+	check(Statement.LHS);
+	FSafeContextScopedEmmitter SafeContextScope(Body, Statement.LHS->Context, *this, TEXT("\t\t\t"));
+
+	Emit(Body, TEXT("\t\t\t"));
+	Emit(Body, *SafeContextScope.GetAdditionalIndent());
+
+	const FString Delegate = TermToText(Statement.LHS);
+
+	Emit(Body, *FString::Printf(TEXT("%s.Clear();\n"), *Delegate));
+}
+
+void FBlueprintCompilerCppBackend::EmitCreateArrayStatement(FKismetFunctionContext& FunctionContext, FBlueprintCompiledStatement& Statement)
+{
+	FBPTerminal* ArrayTerm = Statement.LHS;
+	const FString Array = TermToText(ArrayTerm);
+
+	UArrayProperty* ArrayProperty = CastChecked<UArrayProperty>(ArrayTerm->AssociatedVarProperty);
+	UProperty* InnerProperty = ArrayProperty->Inner;
+
+	for (int32 i = 0; i < Statement.RHS.Num(); ++i)
+	{
+		FBPTerminal* CurrentTerminal = Statement.RHS[i];
+		Emit(Body,
+			*FString::Printf(
+			TEXT("\t\t\t%s[%d] = %s;"),
+			*Array,
+			i,
+			*TermToText(CurrentTerminal, (CurrentTerminal->bIsLiteral ? InnerProperty : NULL))));
+	}
+}
+
+void FBlueprintCompilerCppBackend::EmitGotoStatement(FKismetFunctionContext& FunctionContext, FBlueprintCompiledStatement& Statement)
+{
+	if (Statement.Type == KCST_ComputedGoto)
+	{
+		FString NextStateExpression;
+		NextStateExpression = TermToText(Statement.LHS, (UProperty*)(GetDefault<UIntProperty>()));
+
+		Emit(Body, *FString::Printf(TEXT("\t\t\tCurrentState = %s;\n"), *NextStateExpression));
+		Emit(Body, *FString::Printf(TEXT("\t\t\tbreak;\n")));
+	}
+	else if ((Statement.Type == KCST_GotoIfNot) || (Statement.Type == KCST_EndOfThreadIfNot) || (Statement.Type == KCST_GotoReturnIfNot))
+	{
+		FString ConditionExpression;
+		ConditionExpression = TermToText(Statement.LHS, (UProperty*)(GetDefault<UBoolProperty>()));
+
+		Emit(Body, *FString::Printf(TEXT("\t\t\tif (!%s)\n"), *ConditionExpression));
+		Emit(Body, *FString::Printf(TEXT("\t\t\t{\n")));
+
+		if (Statement.Type == KCST_EndOfThreadIfNot)
+		{
+			ensure(FunctionContext.bUseFlowStack);
+			Emit(Body, TEXT("\t\t\t\tCurrentState = (StateStack.Num() > 0) ? StateStack.Pop(/*bAllowShrinking=*/ false) : -1;\n"));
+		}
+		else if (Statement.Type == KCST_GotoReturnIfNot)
+		{
+			Emit(Body, TEXT("\t\t\t\tCurrentState = -1;\n"));
+		}
+		else
+		{
+			Emit(Body, *FString::Printf(TEXT("\t\t\t\tCurrentState = %d;\n"), StatementToStateIndex(FunctionContext, Statement.TargetLabel)));
+		}
+
+		Emit(Body, *FString::Printf(TEXT("\t\t\t\tbreak;\n")));
+		Emit(Body, *FString::Printf(TEXT("\t\t\t}\n")));
+	}
+	else if (Statement.Type == KCST_GotoReturn)
+	{
+		Emit(Body, TEXT("\t\t\tCurrentState = -1;\n"));
+	}
+	else
+	{
+		Emit(Body, *FString::Printf(TEXT("\t\t\tCurrentState = %d;\n"), StatementToStateIndex(FunctionContext, Statement.TargetLabel)));
+		Emit(Body, *FString::Printf(TEXT("\t\t\tbreak;\n")));
+	}
+}
+
+void FBlueprintCompilerCppBackend::EmitPushStateStatement(FKismetFunctionContext& FunctionContext, FBlueprintCompiledStatement& Statement)
+{
+	ensure(FunctionContext.bUseFlowStack);
+	Emit(Body, *FString::Printf(TEXT("\t\t\tStateStack.Push(%d);\n"), StatementToStateIndex(FunctionContext, Statement.TargetLabel)));
+}
+
+void FBlueprintCompilerCppBackend::EmitEndOfThreadStatement(FKismetFunctionContext& FunctionContext, const FString& ReturnValueString)
+{
+	ensure(FunctionContext.bUseFlowStack);
+	Emit(Body, TEXT("\t\t\tCurrentState = (StateStack.Num() > 0) ? StateStack.Pop(/*bAllowShrinking=*/ false) : -1;\n"));
+	Emit(Body, TEXT("\t\t\tbreak;\n"));
+}
+
+void FBlueprintCompilerCppBackend::EmitReturnStatement(FKismetFunctionContext& FunctionContext, const FString& ReturnValueString)
+{
+	Emit(Body, *FString::Printf(TEXT("\treturn%s;\n"), *ReturnValueString));
+}
+
+void FBlueprintCompilerCppBackend::DeclareLocalVariables(FKismetFunctionContext& FunctionContext, TArray<UProperty*>& LocalVariables)
+{
+	for (int32 i = 0; i < LocalVariables.Num(); ++i)
+	{
+		UProperty* LocalVariable = LocalVariables[i];
+
+		Emit(Body, TEXT("\t"));
+		LocalVariable->ExportCppDeclaration(Body, EExportedDeclaration::Local, NULL, EPropertyExportCPPFlags::CPPF_CustomTypeName);
+		Emit(Body, TEXT("{};\n"));
+	}
+
+	if (LocalVariables.Num() > 0)
+	{
+		Emit(Body, TEXT("\n"));
+	}
+}
+
+void FBlueprintCompilerCppBackend::DeclareStateSwitch(FKismetFunctionContext& FunctionContext)
+{
+	if (FunctionContext.bUseFlowStack)
+	{
+		Emit(Body, TEXT("\tTArray< int32, TInlineAllocator<8> > StateStack;\n"));
+	}
+	Emit(Body, TEXT("\tint32 CurrentState = 0;\n"));
+	Emit(Body, TEXT("\tdo\n"));
+	Emit(Body, TEXT("\t{\n"));
+	Emit(Body, TEXT("\t\tswitch( CurrentState )\n"));
+	Emit(Body, TEXT("\t\t{\n"));
+	Emit(Body, TEXT("\t\tcase 0:\n"));
+}
+
+void FBlueprintCompilerCppBackend::CloseStateSwitch(FKismetFunctionContext& FunctionContext)
+{
+	// Default error-catching case 
+	Emit(Body, TEXT("\t\tdefault:\n"));
+	if (FunctionContext.bUseFlowStack)
+	{
+		Emit(Body, TEXT("\t\t\tcheck(false); // Invalid state\n"));
+	}
+	Emit(Body, TEXT("\t\t\tbreak;\n"));
+
+	// Close the switch block and do-while loop
+	Emit(Body, TEXT("\t\t}\n"));
+	Emit(Body, TEXT("\t} while( CurrentState != -1 );\n"));
+}
+
+void FBlueprintCompilerCppBackend::ConstructFunction(FKismetFunctionContext& FunctionContext, bool bGenerateStubOnly)
+{
+	if (FunctionContext.IsDelegateSignature())
+	{
+		return;
+	}
+
+	UFunction* Function = FunctionContext.Function;
+
+	UProperty* ReturnValue = NULL;
+	TArray<UProperty*> LocalVariables;
+
+	{
+		FString FunctionName;
+		Function->GetName(FunctionName);
+
+		TArray<UProperty*> ArgumentList;
+
+		// Split the function property list into arguments, a return value (if any), and local variable declarations
+		for (TFieldIterator<UProperty> It(Function); It; ++It)
+		{
+			UProperty* Property = *It;
+			if (Property->HasAnyPropertyFlags(CPF_Parm))
+			{
+				if (Property->HasAnyPropertyFlags(CPF_ReturnParm))
+				{
+					if (ReturnValue == NULL)
+					{
+						ReturnValue = Property;
+						LocalVariables.Add(Property);
+					}
+					else
+					{
+						MessageLog.Error(*FString::Printf(TEXT("Function %s from graph @@ has more than one return value (named %s and %s)"),
+							*FunctionName, *ReturnValue->GetName(), *Property->GetName()), FunctionContext.SourceGraph);
+					}
+				}
+				else
+				{
+					ArgumentList.Add(Property);
+				}
+			}
+			else
+			{
+				LocalVariables.Add(Property);
+			}
+		}
+
+		// Emit the declaration
+		const FString ReturnType = ReturnValue ? ReturnValue->GetCPPType(NULL, EPropertyExportCPPFlags::CPPF_CustomTypeName) : TEXT("void");
+
+		//@TODO: Make the header+body export more uniform
+		{
+			const FString Start = FString::Printf(TEXT("%s %s%s%s("), *ReturnType, TEXT("%s"), TEXT("%s"), *FunctionName);
+
+			Emit(Header, *FString::Printf(TEXT("\t%s\n"), *FEmitHelper::EmitUFuntion(Function)));
+			Emit(Header, TEXT("\t"));
+			if (Function->HasAllFunctionFlags(FUNC_Static))
+			{
+				Emit(Header, TEXT("static "));
+			}
+			Emit(Header, *FString::Printf(*Start, TEXT(""), TEXT("")));
+			Emit(Body, *FString::Printf(*Start, *CppClassName, TEXT("::")));
+
+			for (int32 i = 0; i < ArgumentList.Num(); ++i)
+			{
+				UProperty* ArgProperty = ArgumentList[i];
+
+				if (i > 0)
+				{
+					Emit(Header, TEXT(", "));
+					Emit(Body, TEXT(", "));
+				}
+
+				if (ArgProperty->HasAnyPropertyFlags(CPF_OutParm))
+				{
+					Emit(Header, TEXT("/*out*/ "));
+					Emit(Body, TEXT("/*out*/ "));
+				}
+
+				ArgProperty->ExportCppDeclaration(Header, EExportedDeclaration::Parameter, NULL, EPropertyExportCPPFlags::CPPF_CustomTypeName);
+				ArgProperty->ExportCppDeclaration(Body, EExportedDeclaration::Parameter, NULL, EPropertyExportCPPFlags::CPPF_CustomTypeName);
+			}
+
+			Emit(Header, TEXT(")"));
+			Emit(Header, TEXT(";\n"));
+			Emit(Body, TEXT(")\n"));
+		}
+
+		// Start the body of the implementation
+		Emit(Body, TEXT("{\n"));
+	}
+
+	const FString ReturnValueString = ReturnValue ? (FString(TEXT(" ")) + ReturnValue->GetName()) : TEXT("");
+
+	if (!bGenerateStubOnly)
+	{
+		// Emit local variables
+		DeclareLocalVariables(FunctionContext, LocalVariables);
+
+		bool bUseSwitchState = false;
+		for (auto Node : FunctionContext.LinearExecutionList)
+		{
+			TArray<FBlueprintCompiledStatement*>* StatementList = FunctionContext.StatementsPerNode.Find(Node);
+			if (StatementList)
+			{
+				for (auto Statement : (*StatementList))
+				{
+					if (Statement && (
+						Statement->Type == KCST_UnconditionalGoto ||
+						Statement->Type == KCST_PushState ||
+						Statement->Type == KCST_GotoIfNot ||
+						Statement->Type == KCST_ComputedGoto ||
+						Statement->Type == KCST_EndOfThread ||
+						Statement->Type == KCST_EndOfThreadIfNot ||
+						Statement->Type == KCST_GotoReturn ||
+						Statement->Type == KCST_GotoReturnIfNot))
+					{
+						bUseSwitchState = true;
+						break;
+					}
+				}
+			}
+			if (bUseSwitchState)
+			{
+				break;
+			}
+		}
+
+		if (bUseSwitchState)
+		{
+			DeclareStateSwitch(FunctionContext);
+		}
+
+
+		// Run thru code looking only at things marked as jump targets, to make sure the jump targets are ordered in order of appearance in the linear execution list
+		// Emit code in the order specified by the linear execution list (the first node is always the entry point for the function)
+		for (int32 NodeIndex = 0; NodeIndex < FunctionContext.LinearExecutionList.Num(); ++NodeIndex)
+		{
+			UEdGraphNode* StatementNode = FunctionContext.LinearExecutionList[NodeIndex];
+			TArray<FBlueprintCompiledStatement*>* StatementList = FunctionContext.StatementsPerNode.Find(StatementNode);
+
+			if (StatementList != NULL)
+			{
+				for (int32 StatementIndex = 0; StatementIndex < StatementList->Num(); ++StatementIndex)
+				{
+					FBlueprintCompiledStatement& Statement = *((*StatementList)[StatementIndex]);
+
+					if (Statement.bIsJumpTarget)
+					{
+						// Just making sure we number them in order of appearance, so jump statements don't influence the order
+						const int32 StateNum = StatementToStateIndex(FunctionContext, &Statement);
+					}
+				}
+			}
+		}
+
+		// Emit code in the order specified by the linear execution list (the first node is always the entry point for the function)
+		for (int32 NodeIndex = 0; NodeIndex < FunctionContext.LinearExecutionList.Num(); ++NodeIndex)
+		{
+			UEdGraphNode* StatementNode = FunctionContext.LinearExecutionList[NodeIndex];
+			TArray<FBlueprintCompiledStatement*>* StatementList = FunctionContext.StatementsPerNode.Find(StatementNode);
+
+			if (StatementList != NULL)
+			{
+				for (int32 StatementIndex = 0; StatementIndex < StatementList->Num(); ++StatementIndex)
+				{
+					FBlueprintCompiledStatement& Statement = *((*StatementList)[StatementIndex]);
+
+					if (Statement.bIsJumpTarget && bUseSwitchState)
+					{
+						const int32 StateNum = StatementToStateIndex(FunctionContext, &Statement);
+						Emit(Body, *FString::Printf(TEXT("\n\t\tcase %d:\n"), StateNum));
+					}
+
+					switch (Statement.Type)
+					{
+					case KCST_Nop:
+						Emit(Body, TEXT("\t\t\t//No operation.\n"));
+						break;
+					case KCST_WireTraceSite:
+						Emit(Body, TEXT("\t\t\t// Wire debug site.\n"));
+						break;
+					case KCST_DebugSite:
+						Emit(Body, TEXT("\t\t\t// Debug site.\n"));
+						break;
+					case KCST_CallFunction:
+						EmitCallStatment(FunctionContext, Statement);
+						break;
+					case KCST_CallDelegate:
+						EmitCallDelegateStatment(FunctionContext, Statement);
+						break;
+					case KCST_Assignment:
+						EmitAssignmentStatment(FunctionContext, Statement);
+						break;
+					case KCST_CastObjToInterface:
+						EmitCastObjToInterfaceStatement(FunctionContext, Statement);
+						break;
+					case KCST_CrossInterfaceCast:
+						EmitCastBetweenInterfacesStatement(FunctionContext, Statement);
+						break;
+					case KCST_CastInterfaceToObj:
+						EmitCastInterfaceToObjStatement(FunctionContext, Statement);
+						break;
+					case KCST_DynamicCast:
+						EmitDynamicCastStatement(FunctionContext, Statement);
+						break;
+					case KCST_ObjectToBool:
+						EmitObjectToBoolStatement(FunctionContext, Statement);
+						break;
+					case KCST_AddMulticastDelegate:
+						EmitAddMulticastDelegateStatement(FunctionContext, Statement);
+						break;
+					case KCST_RemoveMulticastDelegate:
+						EmitRemoveMulticastDelegateStatement(FunctionContext, Statement);
+						break;
+					case KCST_BindDelegate:
+						EmitBindDelegateStatement(FunctionContext, Statement);
+						break;
+					case KCST_ClearMulticastDelegate:
+						EmitClearMulticastDelegateStatement(FunctionContext, Statement);
+						break;
+					case KCST_CreateArray:
+						EmitCreateArrayStatement(FunctionContext, Statement);
+						break;
+					case KCST_ComputedGoto:
+					case KCST_UnconditionalGoto:
+					case KCST_GotoIfNot:
+					case KCST_EndOfThreadIfNot:
+					case KCST_GotoReturn:
+					case KCST_GotoReturnIfNot:
+						EmitGotoStatement(FunctionContext, Statement);
+						break;
+					case KCST_PushState:
+						EmitPushStateStatement(FunctionContext, Statement);
+						break;
+					case KCST_EndOfThread:
+						EmitEndOfThreadStatement(FunctionContext, ReturnValueString);
+						break;
+					case KCST_Comment:
+						Emit(Body, *FString::Printf(TEXT("\t\t\t// %s\n"), *Statement.Comment));
+						break;
+					case KCST_MetaCast:
+						EmitMetaCastStatement(FunctionContext, Statement);
+						break;
+					case KCST_Return:
+						Emit(Body, TEXT("\t\t\t// Return statement.\n"));
+						break;
+					default:
+						Emit(Body, TEXT("\t// Warning: Ignoring unsupported statement\n"));
+						UE_LOG(LogK2Compiler, Warning, TEXT("C++ backend encountered unsupported statement type %d"), (int32)Statement.Type);
+						break;
+					};
+				}
+			}
+		}
+
+		if (bUseSwitchState)
+		{
+			CloseStateSwitch(FunctionContext);
+		}
+	}
+
+	EmitReturnStatement(FunctionContext, ReturnValueString);
+
+	Emit(Body, TEXT("}\n\n"));
+}
\ No newline at end of file
diff --git a/Engine/Source/Developer/BlueprintCompilerCppBackend/Private/BlueprintCompilerCppBackend.h b/Engine/Source/Developer/BlueprintCompilerCppBackend/Private/BlueprintCompilerCppBackend.h
new file mode 100644
--- /dev/null
+++ b/Engine/Source/Developer/BlueprintCompilerCppBackend/Private/BlueprintCompilerCppBackend.h
@@ -0,0 +1,111 @@
+// Copyright 1998-2015 Epic Games, Inc. All Rights Reserved.
+
+#include "IBlueprintCompilerCppBackendModule.h"
+
+class FBlueprintCompilerCppBackend : public IBlueprintCompilerCppBackend
+{
+protected:
+	FCompilerResultsLog& MessageLog;
+	FKismetCompilerContext& CompilerContext;
+
+	struct FFunctionLabelInfo
+	{
+		TMap<FBlueprintCompiledStatement*, int32> StateMap;
+		int32 StateCounter;
+
+		FFunctionLabelInfo()
+		{
+			StateCounter = 0;
+		}
+
+		int32 StatementToStateIndex(FBlueprintCompiledStatement* Statement)
+		{
+			int32& Index = StateMap.FindOrAdd(Statement);
+			if (Index == 0)
+			{
+				Index = ++StateCounter;
+			}
+
+			return Index;
+		}
+	};
+
+	TArray<FFunctionLabelInfo> StateMapPerFunction;
+	TMap<FKismetFunctionContext*, int32> FunctionIndexMap;
+
+	FString CppClassName;
+
+	// Pointers to commonly used structures (found in constructor)
+	UScriptStruct* VectorStruct;
+	UScriptStruct* RotatorStruct;
+	UScriptStruct* TransformStruct;
+	UScriptStruct* LatentInfoStruct;
+	UScriptStruct* LinearColorStruct;
+public:
+	FStringOutputDevice Header;
+	FStringOutputDevice Body;
+	FString TermToText(const FBPTerminal* Term, const UProperty* SourceProperty = NULL);
+
+	const FString& GetBody()	const override { return Body; }
+	const FString& GetHeader()	const override { return Header; }
+
+protected:
+	FString LatentFunctionInfoTermToText(FBPTerminal* Term, FBlueprintCompiledStatement* TargetLabel);
+
+	int32 StatementToStateIndex(FKismetFunctionContext& FunctionContext, FBlueprintCompiledStatement* Statement)
+	{
+		int32 Index = FunctionIndexMap.FindChecked(&FunctionContext);
+		return StateMapPerFunction[Index].StatementToStateIndex(Statement);
+	}
+public:
+
+	FBlueprintCompilerCppBackend(FKismetCompilerContext& InContext)
+		: MessageLog(InContext.MessageLog)
+		, CompilerContext(InContext)
+	{
+		extern UScriptStruct* Z_Construct_UScriptStruct_FVector();
+		VectorStruct = Z_Construct_UScriptStruct_FVector();
+		RotatorStruct = GetBaseStructure(TEXT("Rotator"));
+		TransformStruct = GetBaseStructure(TEXT("Transform"));
+		LinearColorStruct = GetBaseStructure(TEXT("LinearColor"));
+		LatentInfoStruct = FLatentActionInfo::StaticStruct();
+	}
+
+	void Emit(FStringOutputDevice& Target, const TCHAR* Message)
+	{
+		Target += Message;
+	}
+
+	void EmitClassProperties(FStringOutputDevice& Target, UClass* SourceClass);
+
+	void GenerateCodeFromClass(UClass* SourceClass, FString NewClassName, TIndirectArray<FKismetFunctionContext>& Functions, bool bGenerateStubsOnly) override;
+
+	void EmitCallStatment(FKismetFunctionContext& FunctionContext, FBlueprintCompiledStatement& Statement);
+	void EmitCallDelegateStatment(FKismetFunctionContext& FunctionContext, FBlueprintCompiledStatement& Statement);
+	void EmitAssignmentStatment(FKismetFunctionContext& FunctionContext, FBlueprintCompiledStatement& Statement);
+	void EmitCastObjToInterfaceStatement(FKismetFunctionContext& FunctionContext, FBlueprintCompiledStatement& Statement);
+	void EmitCastBetweenInterfacesStatement(FKismetFunctionContext& FunctionContext, FBlueprintCompiledStatement& Statement);
+	void EmitCastInterfaceToObjStatement(FKismetFunctionContext& FunctionContext, FBlueprintCompiledStatement& Statement);
+	void EmitDynamicCastStatement(FKismetFunctionContext& FunctionContext, FBlueprintCompiledStatement& Statement);
+	void EmitObjectToBoolStatement(FKismetFunctionContext& FunctionContext, FBlueprintCompiledStatement& Statement);
+	void EmitAddMulticastDelegateStatement(FKismetFunctionContext& FunctionContext, FBlueprintCompiledStatement& Statement);
+	void EmitBindDelegateStatement(FKismetFunctionContext& FunctionContext, FBlueprintCompiledStatement& Statement);
+	void EmitClearMulticastDelegateStatement(FKismetFunctionContext& FunctionContext, FBlueprintCompiledStatement& Statement);
+	void EmitCreateArrayStatement(FKismetFunctionContext& FunctionContext, FBlueprintCompiledStatement& Statement);
+	void EmitGotoStatement(FKismetFunctionContext& FunctionContext, FBlueprintCompiledStatement& Statement);
+	void EmitPushStateStatement(FKismetFunctionContext& FunctionContext, FBlueprintCompiledStatement& Statement);
+	void EmitEndOfThreadStatement(FKismetFunctionContext& FunctionContext, const FString& ReturnValueString);
+	void EmitReturnStatement(FKismetFunctionContext& FunctionContext, const FString& ReturnValueString);
+	void EmitRemoveMulticastDelegateStatement(FKismetFunctionContext& FunctionContext, FBlueprintCompiledStatement& Statement);
+	void EmitMetaCastStatement(FKismetFunctionContext& FunctionContext, FBlueprintCompiledStatement& Statement);
+
+	/** Emits local variable declarations for a function */
+	void DeclareLocalVariables(FKismetFunctionContext& FunctionContext, TArray<UProperty*>& LocalVariables);
+
+	/** Emits the internal execution flow state declaration for a function */
+	void DeclareStateSwitch(FKismetFunctionContext& FunctionContext);
+	void CloseStateSwitch(FKismetFunctionContext& FunctionContext);
+
+	/** Builds both the header declaration and body implementation of a function */
+	void ConstructFunction(FKismetFunctionContext& FunctionContext, bool bGenerateStubOnly);
+};
\ No newline at end of file
diff --git a/Engine/Source/Developer/BlueprintCompilerCppBackend/Private/BlueprintCompilerCppBackendModule.cpp b/Engine/Source/Developer/BlueprintCompilerCppBackend/Private/BlueprintCompilerCppBackendModule.cpp
new file mode 100644
--- /dev/null
+++ b/Engine/Source/Developer/BlueprintCompilerCppBackend/Private/BlueprintCompilerCppBackendModule.cpp
@@ -0,0 +1,34 @@
+// Copyright 1998-2015 Epic Games, Inc. All Rights Reserved.
+
+#include "BlueprintCompilerCppBackendModulePrivatePCH.h"
+#include "IBlueprintCompilerCppBackendModule.h"
+#include "ModuleManager.h"
+
+class FBlueprintCompilerCppBackendModule : public IBlueprintCompilerCppBackendModule
+{
+	/** IModuleInterface implementation */
+	virtual void StartupModule() override;
+	virtual void ShutdownModule() override;
+	virtual IBlueprintCompilerCppBackend* Create(FKismetCompilerContext& InContext)  override;
+};
+
+IMPLEMENT_MODULE(FBlueprintCompilerCppBackendModule, BlueprintCompilerCppBackend)
+
+
+
+void FBlueprintCompilerCppBackendModule::StartupModule()
+{
+	// This code will execute after your module is loaded into memory (but after global variables are initialized, of course.)
+}
+
+
+void FBlueprintCompilerCppBackendModule::ShutdownModule()
+{
+	// This function may be called during shutdown to clean up your module.  For modules that support dynamic reloading,
+	// we call this function before unloading the module.
+}
+
+IBlueprintCompilerCppBackend* FBlueprintCompilerCppBackendModule::Create(FKismetCompilerContext& InContext)
+{
+	return new FBlueprintCompilerCppBackend(InContext);
+}
diff --git a/Engine/Source/Developer/BlueprintCompilerCppBackend/Private/BlueprintCompilerCppBackendModulePrivatePCH.h b/Engine/Source/Developer/BlueprintCompilerCppBackend/Private/BlueprintCompilerCppBackendModulePrivatePCH.h
new file mode 100644
--- /dev/null
+++ b/Engine/Source/Developer/BlueprintCompilerCppBackend/Private/BlueprintCompilerCppBackendModulePrivatePCH.h
@@ -0,0 +1,10 @@
+// Copyright 1998-2015 Epic Games, Inc. All Rights Reserved.
+
+#include "CoreUObject.h"
+#include "Engine.h"
+#include "KismetCompiler.h"
+#include "SourceCodeNavigation.h"
+#include "DefaultValueHelper.h"
+
+// You should place include statements to your module's private header files here.  You only need to
+// add includes for headers that are used in most of your module's source files though.
diff --git a/Engine/Source/Developer/BlueprintCompilerCppBackend/Private/BlueprintCompilerCppBackendUtils.h b/Engine/Source/Developer/BlueprintCompilerCppBackend/Private/BlueprintCompilerCppBackendUtils.h
new file mode 100644
--- /dev/null
+++ b/Engine/Source/Developer/BlueprintCompilerCppBackend/Private/BlueprintCompilerCppBackendUtils.h
@@ -0,0 +1,424 @@
+// Copyright 1998-2015 Epic Games, Inc. All Rights Reserved.
+
+// Generates single if scope. It's condition checks context of given term.
+struct FSafeContextScopedEmmitter
+{
+private:
+	FStringOutputDevice& Body;
+	bool bSafeContextUsed;
+	const TCHAR* CurrentIndent;
+public:
+	FString GetAdditionalIndent() const
+	{
+		return bSafeContextUsed ? TEXT("\t") : FString();
+	}
+
+	bool IsSafeContextUsed() const
+	{
+		return bSafeContextUsed;
+	}
+
+	FSafeContextScopedEmmitter(FStringOutputDevice& InBody, const FBPTerminal* Term, FBlueprintCompilerCppBackend& CppBackend, const TCHAR* InCurrentIndent)
+		: Body(InBody), bSafeContextUsed(false), CurrentIndent(InCurrentIndent)
+	{
+		TArray<FString> SafetyConditions;
+		for (; Term; Term = Term->Context)
+		{
+			if (!Term->bIsStructContext && (Term->Type.PinSubCategory != TEXT("self")))
+			{
+				ensure(!Term->bIsLiteral);
+				SafetyConditions.Add(CppBackend.TermToText(Term));
+			}
+		}
+
+		if (SafetyConditions.Num())
+		{
+			bSafeContextUsed = true;
+			Body += CurrentIndent;
+			Body += TEXT("if (");
+			for (int32 Iter = SafetyConditions.Num() - 1; Iter >= 0; --Iter)
+			{
+				Body += FString(TEXT("IsValid("));
+				Body += SafetyConditions[Iter];
+				Body += FString(TEXT(")"));
+				if (Iter)
+				{
+					Body += FString(TEXT(" && "));
+				}
+			}
+			Body += TEXT(")\n");
+			Body += CurrentIndent;
+			Body += TEXT("{\n");
+		}
+	}
+
+	~FSafeContextScopedEmmitter()
+	{
+		if (bSafeContextUsed)
+		{
+			Body += CurrentIndent;
+			Body += TEXT("}\n");
+		}
+	}
+};
+
+struct FEmitHelper
+{
+	static void ArrayToString(const TArray<FString>& Array, FString& String, const TCHAR* Separator)
+	{
+		if (Array.Num())
+		{
+			String += Array[0];
+		}
+		for (int32 i = 1; i < Array.Num(); ++i)
+		{
+			String += TEXT(", ");
+			String += Array[i];
+		}
+	}
+
+	static bool HasAllFlags(uint64 Flags, uint64 FlagsToCheck)
+	{
+		return FlagsToCheck == (Flags & FlagsToCheck);
+	}
+
+	static FString HandleRepNotifyFunc(const UProperty* Property)
+	{
+		check(Property);
+		if (HasAllFlags(Property->PropertyFlags, CPF_Net | CPF_RepNotify))
+		{
+			return FString::Printf(TEXT("ReplicatedUsing=%s"), *Property->RepNotifyFunc.ToString());
+		}
+		else if (HasAllFlags(Property->PropertyFlags, CPF_Net))
+		{
+			return TEXT("Replicated");
+		}
+		return FString();
+	}
+
+	static bool MetaDataCanBeNative(const FName MetaDataName)
+	{
+		if (MetaDataName == TEXT("ModuleRelativePath"))
+		{
+			return false;
+		}
+		return true;
+	}
+
+	static FString HandleMetaData(const UField* Field, bool AddCategory = true)
+	{
+		FString MetaDataStr;
+
+		check(Field);
+		UPackage* Package = Field->GetOutermost();
+		check(Package);
+		const UMetaData* MetaData = Package->GetMetaData();
+		check(MetaData);
+		const TMap<FName, FString>* ValuesMap = MetaData->ObjectMetaDataMap.Find(Field);
+		TArray<FString> MetaDataStrings;
+		if (ValuesMap && ValuesMap->Num())
+		{
+			for (auto& Pair : *ValuesMap)
+			{
+				if (!MetaDataCanBeNative(Pair.Key))
+				{
+					continue;
+				}
+				if (!Pair.Value.IsEmpty())
+				{
+					FString Value = Pair.Value.Replace(TEXT("\n"), TEXT(""));
+					MetaDataStrings.Emplace(FString::Printf(TEXT("%s=\"%s\""), *Pair.Key.ToString(), *Pair.Value));
+				}
+				else
+				{
+					MetaDataStrings.Emplace(Pair.Key.ToString());
+				}
+			}
+		}
+		if (AddCategory && (!ValuesMap || !ValuesMap->Find(TEXT("Category"))))
+		{
+			MetaDataStrings.Emplace(TEXT("Category"));
+		}
+		MetaDataStrings.Remove(FString());
+		if (MetaDataStrings.Num())
+		{
+			MetaDataStr += TEXT("meta=(");
+			ArrayToString(MetaDataStrings, MetaDataStr, TEXT(", "));
+			MetaDataStr += TEXT(")");
+		}
+		return MetaDataStr;
+	}
+
+#ifdef HANDLE_CPF_TAG
+	static_assert(false, "Macro HANDLE_CPF_TAG redefinition.");
+#endif
+#define HANDLE_CPF_TAG(TagName, CheckedFlags) if (HasAllFlags(Flags, (CheckedFlags))) { Tags.Emplace(TagName); }
+
+	static TArray<FString> ProperyFlagsToTags(uint64 Flags)
+	{
+		TArray<FString> Tags;
+
+		// EDIT FLAGS
+		if (HasAllFlags(Flags, (CPF_Edit | CPF_EditConst | CPF_DisableEditOnInstance)))
+		{
+			Tags.Emplace(TEXT("VisibleDefaultsOnly"));
+		}
+		else if (HasAllFlags(Flags, (CPF_Edit | CPF_EditConst | CPF_DisableEditOnTemplate)))
+		{
+			Tags.Emplace(TEXT("VisibleInstanceOnly"));
+		}
+		else if (HasAllFlags(Flags, (CPF_Edit | CPF_EditConst)))
+		{
+			Tags.Emplace(TEXT("VisibleAnywhere"));
+		}
+		else if (HasAllFlags(Flags, (CPF_Edit | CPF_DisableEditOnInstance)))
+		{
+			Tags.Emplace(TEXT("EditDefaultsOnly"));
+		}
+		else if (HasAllFlags(Flags, (CPF_Edit | CPF_DisableEditOnTemplate)))
+		{
+			Tags.Emplace(TEXT("EditInstanceOnly"));
+		}
+		else if (HasAllFlags(Flags, (CPF_Edit)))
+		{
+			Tags.Emplace(TEXT("EditAnywhere"));
+		}
+
+		// BLUEPRINT EDIT
+		if (HasAllFlags(Flags, (CPF_BlueprintVisible | CPF_BlueprintReadOnly)))
+		{
+			Tags.Emplace(TEXT("BlueprintReadOnly"));
+		}
+		else if (HasAllFlags(Flags, (CPF_BlueprintVisible)))
+		{
+			Tags.Emplace(TEXT("BlueprintReadWrite"));
+		}
+
+		// CONFIG
+		if (HasAllFlags(Flags, (CPF_GlobalConfig | CPF_Config)))
+		{
+			Tags.Emplace(TEXT("GlobalConfig"));
+		}
+		else if (HasAllFlags(Flags, (CPF_Config)))
+		{
+			Tags.Emplace(TEXT("Config"));
+		}
+
+		// OTHER
+		HANDLE_CPF_TAG(TEXT("Transient"), CPF_Transient)
+			HANDLE_CPF_TAG(TEXT("DuplicateTransient"), CPF_DuplicateTransient)
+			HANDLE_CPF_TAG(TEXT("TextExportTransient"), CPF_TextExportTransient)
+			HANDLE_CPF_TAG(TEXT("NonPIEDuplicateTransient"), CPF_NonPIEDuplicateTransient)
+			HANDLE_CPF_TAG(TEXT("Export"), CPF_ExportObject)
+			HANDLE_CPF_TAG(TEXT("NoClear"), CPF_NoClear)
+			HANDLE_CPF_TAG(TEXT("EditFixedSize"), CPF_EditFixedSize)
+			HANDLE_CPF_TAG(TEXT("NotReplicated"), CPF_RepSkip)
+			HANDLE_CPF_TAG(TEXT("RepRetry"), CPF_RepRetry)
+			HANDLE_CPF_TAG(TEXT("Interp"), CPF_Edit | CPF_BlueprintVisible | CPF_Interp)
+			HANDLE_CPF_TAG(TEXT("NonTransactional"), CPF_NonTransactional)
+			HANDLE_CPF_TAG(TEXT("BlueprintAssignable"), CPF_BlueprintAssignable)
+			HANDLE_CPF_TAG(TEXT("BlueprintCallable"), CPF_BlueprintCallable)
+			HANDLE_CPF_TAG(TEXT("BlueprintAuthorityOnly"), CPF_BlueprintAuthorityOnly)
+			HANDLE_CPF_TAG(TEXT("AssetRegistrySearchable"), CPF_AssetRegistrySearchable)
+			HANDLE_CPF_TAG(TEXT("SimpleDisplay"), CPF_SimpleDisplay)
+			HANDLE_CPF_TAG(TEXT("AdvancedDisplay"), CPF_AdvancedDisplay)
+			HANDLE_CPF_TAG(TEXT("SaveGame"), CPF_SaveGame)
+
+			//TODO:
+			//HANDLE_CPF_TAG(TEXT("Instanced"), CPF_PersistentInstance | CPF_ExportObject | CPF_InstancedReference)
+
+			return Tags;
+	}
+
+	static TArray<FString> FunctionFlagsToTags(uint64 Flags)
+	{
+		TArray<FString> Tags;
+
+		//Pointless: BlueprintNativeEvent, BlueprintImplementableEvent
+		//Pointless: CustomThunk ?
+
+		//TODO: SealedEvent
+		//TODO: Unreliable
+		//TODO: ServiceRequest, ServiceResponse
+
+		HANDLE_CPF_TAG(TEXT("Exec"), FUNC_Exec)
+			HANDLE_CPF_TAG(TEXT("Server"), FUNC_Net | FUNC_NetServer)
+			HANDLE_CPF_TAG(TEXT("Client"), FUNC_Net | FUNC_NetClient)
+			HANDLE_CPF_TAG(TEXT("NetMulticast"), FUNC_Net | FUNC_NetMulticast)
+			HANDLE_CPF_TAG(TEXT("Reliable"), FUNC_NetReliable)
+			HANDLE_CPF_TAG(TEXT("BlueprintCallable"), FUNC_BlueprintCallable)
+			HANDLE_CPF_TAG(TEXT("BlueprintPure"), FUNC_BlueprintCallable | FUNC_BlueprintPure)
+			HANDLE_CPF_TAG(TEXT("BlueprintAuthorityOnly"), FUNC_BlueprintAuthorityOnly)
+			HANDLE_CPF_TAG(TEXT("BlueprintCosmetic"), FUNC_BlueprintCosmetic)
+			HANDLE_CPF_TAG(TEXT("WithValidation"), FUNC_NetValidate)
+
+			return Tags;
+	}
+
+#undef HANDLE_CPF_TAG
+
+	static bool IsBlueprintNativeEvent(uint64 FunctionFlags)
+	{
+		return HasAllFlags(FunctionFlags, FUNC_Event | FUNC_BlueprintEvent | FUNC_Native);
+	}
+
+	static bool IsBlueprintImplementableEvent(uint64 FunctionFlags)
+	{
+		return HasAllFlags(FunctionFlags, FUNC_Event | FUNC_BlueprintEvent) && !HasAllFlags(FunctionFlags, FUNC_Native);
+	}
+
+	static FString EmitUFuntion(UFunction* Function)
+	{
+		TArray<FString> Tags = FEmitHelper::FunctionFlagsToTags(Function->FunctionFlags);
+		const bool bMustHaveCategory = (Function->FunctionFlags & (FUNC_BlueprintCallable | FUNC_BlueprintPure)) != 0;
+		Tags.Emplace(FEmitHelper::HandleMetaData(Function, bMustHaveCategory));
+		Tags.Remove(FString());
+
+		FString AllTags;
+		FEmitHelper::ArrayToString(Tags, AllTags, TEXT(", "));
+
+		return FString::Printf(TEXT("UFUNCTION(%s)"), *AllTags);
+	}
+
+	static int32 ParseDelegateDetails(UFunction* Signature, FString& OutParametersMacro, FString& OutParamNumberStr)
+	{
+		int32 ParameterNum = 0;
+		FStringOutputDevice Parameters;
+		for (TFieldIterator<UProperty> PropIt(Signature); PropIt && (PropIt->PropertyFlags & CPF_Parm); ++PropIt)
+		{
+			Parameters += ", ";
+			PropIt->ExportCppDeclaration(Parameters, EExportedDeclaration::MacroParameter, NULL, EPropertyExportCPPFlags::CPPF_CustomTypeName);
+			++ParameterNum;
+		}
+
+		FString ParamNumberStr;
+		switch (ParameterNum)
+		{
+		case 0: ParamNumberStr = TEXT("");				break;
+		case 1: ParamNumberStr = TEXT("_OneParam");		break;
+		case 2: ParamNumberStr = TEXT("_TwoParams");	break;
+		case 3: ParamNumberStr = TEXT("_ThreeParams");	break;
+		case 4: ParamNumberStr = TEXT("_FourParams");	break;
+		case 5: ParamNumberStr = TEXT("_FiveParams");	break;
+		case 6: ParamNumberStr = TEXT("_SixParams");	break;
+		case 7: ParamNumberStr = TEXT("_SevenParams");	break;
+		case 8: ParamNumberStr = TEXT("_EightParams");	break;
+		default: ParamNumberStr = TEXT("_TooMany");		break;
+		}
+
+		OutParametersMacro = Parameters;
+		OutParamNumberStr = ParamNumberStr;
+		return ParameterNum;
+	}
+
+	static TArray<FString> EmitSinglecastDelegateDeclarations(const TArray<UDelegateProperty*>& Delegates)
+	{
+		TArray<FString> Results;
+		for (auto It : Delegates)
+		{
+			check(It);
+			auto Signature = It->SignatureFunction;
+			check(Signature);
+
+			FString ParamNumberStr, Parameters;
+			ParseDelegateDetails(Signature, Parameters, ParamNumberStr);
+
+			Results.Add(*FString::Printf(TEXT("DECLARE_DYNAMIC_DELEGATE%s(%s%s)"),
+				*ParamNumberStr, *It->GetCPPType(NULL, EPropertyExportCPPFlags::CPPF_CustomTypeName), *Parameters));
+		}
+		return Results;
+	}
+
+	static TArray<FString> EmitMulticastDelegateDeclarations(UClass* SourceClass)
+	{
+		TArray<FString> Results;
+		for (TFieldIterator<UMulticastDelegateProperty> It(SourceClass, EFieldIteratorFlags::ExcludeSuper); It; ++It)
+		{
+			auto Signature = It->SignatureFunction;
+			check(Signature);
+
+			FString ParamNumberStr, Parameters;
+			ParseDelegateDetails(Signature, Parameters, ParamNumberStr);
+
+			Results.Add(*FString::Printf(TEXT("DECLARE_DYNAMIC_MULTICAST_DELEGATE%s(%s%s)"),
+				*ParamNumberStr, *It->GetCPPType(NULL, EPropertyExportCPPFlags::CPPF_CustomTypeName), *Parameters));
+		}
+
+		return Results;
+	}
+
+	static FString EmitLifetimeReplicatedPropsImpl(UClass* SourceClass, const FString& CppClassName, const TCHAR* InCurrentIndent)
+	{
+		FString Result;
+		bool bFunctionInitilzed = false;
+		for (TFieldIterator<UProperty> It(SourceClass, EFieldIteratorFlags::ExcludeSuper); It; ++It)
+		{
+			if ((It->PropertyFlags & CPF_Net) != 0)
+			{
+				if (!bFunctionInitilzed)
+				{
+					Result += FString::Printf(TEXT("%svoid %s::GetLifetimeReplicatedProps(TArray< FLifetimeProperty > & OutLifetimeProps) const\n"), InCurrentIndent, *CppClassName);
+					Result += FString::Printf(TEXT("%s{\n"), InCurrentIndent);
+					Result += FString::Printf(TEXT("%s\tSuper::GetLifetimeReplicatedProps(OutLifetimeProps);\n"), InCurrentIndent);
+					bFunctionInitilzed = true;
+				}
+				Result += FString::Printf(TEXT("%s\tDOREPLIFETIME( %s, %s);\n"), InCurrentIndent, *CppClassName, *It->GetNameCPP());
+			}
+		}
+		if (bFunctionInitilzed)
+		{
+			Result += FString::Printf(TEXT("%s}\n"), InCurrentIndent);
+		}
+		return Result;
+	}
+
+	static FString GatherHeadersToInclude(UClass* SourceClass, const TArray<FString>& PersistentHeaders)
+	{
+		TSet<FString> HeaderFiles;
+		HeaderFiles.Append(PersistentHeaders);
+		{
+			TArray<UObject*> ReferencedObjects;
+			{
+				FReferenceFinder ReferenceFinder(ReferencedObjects, NULL, false, false, false, true);
+				TArray<UObject*> ObjectsToCheck;
+				GetObjectsWithOuter(SourceClass, ObjectsToCheck, true);
+				//CDO?
+				for (auto Obj : ObjectsToCheck)
+				{
+					ReferenceFinder.FindReferences(Obj);
+				}
+			}
+
+			auto EngineSourceDir = FPaths::EngineSourceDir();
+			auto GameSourceDir = FPaths::GameSourceDir();
+			auto CurrentPackage = SourceClass->GetTypedOuter<UPackage>();
+			for (auto Obj : ReferencedObjects)
+			{
+				auto Field = Cast<UField>(Obj);
+				FString PackPath;
+				if (Field
+					&& (Field->GetTypedOuter<UPackage>() != CurrentPackage)
+					&& FSourceCodeNavigation::FindClassHeaderPath(Field, PackPath))
+				{
+					if (!PackPath.RemoveFromStart(EngineSourceDir))
+					{
+						if (!PackPath.RemoveFromStart(GameSourceDir))
+						{
+							PackPath = FPaths::GetCleanFilename(PackPath);
+						}
+					}
+					HeaderFiles.Add(PackPath);
+				}
+			}
+		}
+
+		FString Result;
+		for (auto HeaderFile : HeaderFiles)
+		{
+			Result += FString::Printf(TEXT("#include \"%s\"\n"), *HeaderFile);
+		}
+
+		return Result;
+	}
+};
+
diff --git a/Engine/Source/Developer/BlueprintCompilerCppBackend/Public/IBlueprintCompilerCppBackendModule.h b/Engine/Source/Developer/BlueprintCompilerCppBackend/Public/IBlueprintCompilerCppBackendModule.h
new file mode 100644
--- /dev/null
+++ b/Engine/Source/Developer/BlueprintCompilerCppBackend/Public/IBlueprintCompilerCppBackendModule.h
@@ -0,0 +1,14 @@
+// Copyright 1998-2015 Epic Games, Inc. All Rights Reserved.
+
+#pragma once
+
+#include "Editor/KismetCompiler/Public/BlueprintCompilerCppBackendInterface.h"
+
+/**
+ * The public interface to this module
+ */
+class IBlueprintCompilerCppBackendModule : public IBlueprintCompilerCppBackendModuleInterface
+{
+
+};
+
diff --git a/Engine/Source/Editor/Kismet/Kismet.Build.cs b/Engine/Source/Editor/Kismet/Kismet.Build.cs
--- a/Engine/Source/Editor/Kismet/Kismet.Build.cs
+++ b/Engine/Source/Editor/Kismet/Kismet.Build.cs
@@ -56,6 +56,7 @@ public class Kismet : ModuleRules
 				"Documentation",
 				"EditorWidgets",
 				"GameProjectGeneration",
+                "BlueprintCompilerCppBackend",
 			}
             );
 	}
diff --git a/Engine/Source/Editor/KismetCompiler/Private/KismetCompiler.cpp b/Engine/Source/Editor/KismetCompiler/Private/KismetCompiler.cpp
--- a/Engine/Source/Editor/KismetCompiler/Private/KismetCompiler.cpp
+++ b/Engine/Source/Editor/KismetCompiler/Private/KismetCompiler.cpp
@@ -24,6 +24,7 @@
 #include "Engine/DynamicBlueprintBinding.h"
 #include "Kismet/KismetMathLibrary.h"
 #include "Engine/InheritableComponentHandler.h"
+#include "BlueprintCompilerCppBackendInterface.h"
 
 static bool bDebugPropertyPropagation = false;
 
@@ -3576,7 +3577,7 @@ void FKismetCompilerContext::Compile()
 		// Generate code thru the backend(s)
 		if ((bDisplayCpp && bIsFullCompile) || CompileOptions.DoesRequireCppCodeGeneration())
 		{
-			TUniquePtr<IKismetCppBackend> Backend_CPP(IKismetCppBackend::Create(Schema, *this));
+			TUniquePtr<IBlueprintCompilerCppBackend> Backend_CPP(IBlueprintCompilerCppBackendModuleInterface::Get().Create(*this));
 
 			// The C++ backend is currently only for debugging, so it's only run if the output will be visible
 			Backend_CPP->GenerateCodeFromClass(NewClass, CompileOptions.NewCppClassName, FunctionList, !bIsFullCompile);
diff --git a/Engine/Source/Editor/KismetCompiler/Private/KismetCompilerBackend.h b/Engine/Source/Editor/KismetCompiler/Private/KismetCompilerBackend.h
--- a/Engine/Source/Editor/KismetCompiler/Private/KismetCompilerBackend.h
+++ b/Engine/Source/Editor/KismetCompiler/Private/KismetCompilerBackend.h
@@ -15,21 +15,6 @@ class IKismetCompilerBackend
 {
 };
 
-//////////////////////////////////////////////////////////////////////////
-// FKismetCppBackend
-
-class IKismetCppBackend : public IKismetCompilerBackend
-{
-public:
-	static IKismetCppBackend* Create(UEdGraphSchema_K2* InSchema, FKismetCompilerContext& InContext);
-
-	virtual const FString& GetBody() const = 0;
-	virtual const FString& GetHeader() const = 0;
-	virtual void GenerateCodeFromClass(UClass* SourceClass, FString NewClassName, TIndirectArray<FKismetFunctionContext>& Functions, bool bGenerateStubsOnly = false) = 0;
-
-	virtual ~IKismetCppBackend() {}
-};
-
 //////////////////////////////////////////////////////////////////////////
 // FKismetVMBackend
 
diff --git a/Engine/Source/Editor/KismetCompiler/Private/KismetCompilerCppBackend.cpp b/Engine/Source/Editor/KismetCompiler/Private/KismetCompilerCppBackend.cpp
deleted file mode 100644
diff --git a/Engine/Source/Editor/KismetCompiler/Public/BlueprintCompilerCppBackendInterface.h b/Engine/Source/Editor/KismetCompiler/Public/BlueprintCompilerCppBackendInterface.h
new file mode 100644
--- /dev/null
+++ b/Engine/Source/Editor/KismetCompiler/Public/BlueprintCompilerCppBackendInterface.h
@@ -0,0 +1,48 @@
+// Copyright 1998-2015 Epic Games, Inc. All Rights Reserved.
+
+#pragma once
+
+#include "ModuleManager.h"
+#include "KismetCompiler.h"
+
+class IBlueprintCompilerCppBackend
+{
+public:
+	virtual void GenerateCodeFromClass(UClass* SourceClass, FString NewClassName, TIndirectArray<FKismetFunctionContext>& Functions, bool bGenerateStubsOnly) = 0;
+
+	virtual const FString& GetBody() const = 0;
+	virtual const FString& GetHeader() const = 0;
+
+	virtual ~IBlueprintCompilerCppBackend() {}
+};
+
+/**
+ * BlueprintCompilerCppBackend module interface
+ */
+class IBlueprintCompilerCppBackendModuleInterface : public IModuleInterface
+{
+public:
+
+	/**
+	 * Singleton-like access to this module's interface.  This is just for convenience!
+	 * Beware of calling this during the shutdown phase, though.  Your module might have been unloaded already.
+	 *
+	 * @return Returns singleton instance, loading the module on demand if needed
+	 */
+	static inline IBlueprintCompilerCppBackendModuleInterface& Get()
+	{
+		return FModuleManager::LoadModuleChecked< IBlueprintCompilerCppBackendModuleInterface >("BlueprintCompilerCppBackend");
+	}
+
+	/**
+	 * Checks to see if this module is loaded and ready.  It is only valid to call Get() if IsAvailable() returns true.
+	 *
+	 * @return True if the module is loaded and ready to use
+	 */
+	static inline bool IsAvailable()
+	{
+		return FModuleManager::Get().IsModuleLoaded("BlueprintCompilerCppBackend");
+	}
+
+	virtual IBlueprintCompilerCppBackend* Create(FKismetCompilerContext& InContext) = 0;
+};
\ No newline at end of file
