Fixed non-unity compile in BlueprintCompilerCppBackendModule.cpp

diff --git a/Engine/Source/Developer/BlueprintCompilerCppBackend/Private/BlueprintCompilerCppBackendModule.cpp b/Engine/Source/Developer/BlueprintCompilerCppBackend/Private/BlueprintCompilerCppBackendModule.cpp
--- a/Engine/Source/Developer/BlueprintCompilerCppBackend/Private/BlueprintCompilerCppBackendModule.cpp
+++ b/Engine/Source/Developer/BlueprintCompilerCppBackend/Private/BlueprintCompilerCppBackendModule.cpp
@@ -3,6 +3,7 @@
 #include "BlueprintCompilerCppBackendModulePrivatePCH.h"
 #include "IBlueprintCompilerCppBackendModule.h"
 #include "ModuleManager.h"
+#include "BlueprintCompilerCppBackend.h"
 
 class FBlueprintCompilerCppBackendModule : public IBlueprintCompilerCppBackendModule
 {
