added missing #pragma once to BlueprintCompilerBackend.h

diff --git a/Engine/Source/Developer/BlueprintCompilerCppBackend/Private/BlueprintCompilerCppBackend.h b/Engine/Source/Developer/BlueprintCompilerCppBackend/Private/BlueprintCompilerCppBackend.h
--- a/Engine/Source/Developer/BlueprintCompilerCppBackend/Private/BlueprintCompilerCppBackend.h
+++ b/Engine/Source/Developer/BlueprintCompilerCppBackend/Private/BlueprintCompilerCppBackend.h
@@ -1,4 +1,5 @@
 // Copyright 1998-2015 Epic Games, Inc. All Rights Reserved.
+#pragma  once
 
 #include "IBlueprintCompilerCppBackendModule.h"
 
