C++ Backend - support for SwitchValue and inline statements

diff --git a/Engine/Source/Developer/BlueprintCompilerCppBackend/Private/BlueprintCompilerCppBackend.cpp b/Engine/Source/Developer/BlueprintCompilerCppBackend/Private/BlueprintCompilerCppBackend.cpp
--- a/Engine/Source/Developer/BlueprintCompilerCppBackend/Private/BlueprintCompilerCppBackend.cpp
+++ b/Engine/Source/Developer/BlueprintCompilerCppBackend/Private/BlueprintCompilerCppBackend.cpp
@@ -56,21 +56,21 @@ FString FBlueprintCompilerCppBackend::TermToText(const FBPTerminal* Term, const
 		}
 		else if (auto StructProperty = Cast<const UStructProperty>(CoerceProperty))
 		{
-			if (StructProperty->Struct == VectorStruct)
+			if (StructProperty->Struct == TBaseStructure<FVector>::Get())
 			{
 				FVector Vect = FVector::ZeroVector;
 				FDefaultValueHelper::ParseVector(Term->Name, /*out*/ Vect);
 
 				return FString::Printf(TEXT("FVector(%f,%f,%f)"), Vect.X, Vect.Y, Vect.Z);
 			}
-			else if (StructProperty->Struct == RotatorStruct)
+			else if (StructProperty->Struct == TBaseStructure<FRotator>::Get())
 			{
 				FRotator Rot = FRotator::ZeroRotator;
 				FDefaultValueHelper::ParseRotator(Term->Name, /*out*/ Rot);
 
 				return FString::Printf(TEXT("FRotator(%f,%f,%f)"), Rot.Pitch, Rot.Yaw, Rot.Roll);
 			}
-			else if (StructProperty->Struct == TransformStruct)
+			else if (StructProperty->Struct == TBaseStructure<FTransform>::Get())
 			{
 				FTransform Trans = FTransform::Identity;
 				Trans.InitFromString(Term->Name);
@@ -82,7 +82,7 @@ FString FBlueprintCompilerCppBackend::TermToText(const FBPTerminal* Term, const
 				return FString::Printf(TEXT("FTransform( FQuat(%f,%f,%f,%f), FVector(%f,%f,%f), FVector(%f,%f,%f) )"),
 					Rot.X, Rot.Y, Rot.Z, Rot.W, Translation.X, Translation.Y, Translation.Z, Scale.X, Scale.Y, Scale.Z);
 			}
-			else if (StructProperty->Struct == LinearColorStruct)
+			else if (StructProperty->Struct == TBaseStructure<FLinearColor>::Get())
 			{
 				FLinearColor LinearColor;
 				LinearColor.InitFromString(Term->Name);
@@ -102,18 +102,6 @@ FString FBlueprintCompilerCppBackend::TermToText(const FBPTerminal* Term, const
 			}
 			return FString(TEXT("NULL"));
 		}
-		else if (CoerceProperty->IsA(UDelegateProperty::StaticClass()))
-		{
-			//@TODO: K2 Delegate Support: Won't compile, there isn't an operator= that does what we want here, it should be a proper call to BindDynamic instead!
-			if (Term->Name == TEXT(""))
-			{
-				return TEXT("NULL");
-			}
-			else
-			{
-				return FString::Printf(TEXT("BindDynamic(this, &%s::%s)"), *CppClassName, *(Term->Name));
-			}
-		}
 		else if (CoerceProperty->IsA(UObjectPropertyBase::StaticClass()))
 		{
 			if (Term->Type.PinSubCategory == PSC_Self)
@@ -159,6 +147,22 @@ FString FBlueprintCompilerCppBackend::TermToText(const FBPTerminal* Term, const
 			return Term->Name;
 		}
 	}
+	else if (Term->InlineGeneratedParameter)
+	{
+		if (KCST_SwitchValue == Term->InlineGeneratedParameter->Type)
+		{
+			return EmitSwitchValueStatmentInner(*Term->InlineGeneratedParameter);
+		}
+		else if (KCST_CallFunction == Term->InlineGeneratedParameter->Type)
+		{
+			return EmitCallStatmentInner(*Term->InlineGeneratedParameter, true);
+		}
+		else
+		{
+			ensureMsgf(false, TEXT("KCST %d is not accepted as inline statement."), Term->InlineGeneratedParameter->Type);
+			return FString();
+		}
+	}
 	else
 	{
 		FString Prefix(TEXT(""));
@@ -183,7 +187,7 @@ FString FBlueprintCompilerCppBackend::TermToText(const FBPTerminal* Term, const
 
 FString FBlueprintCompilerCppBackend::LatentFunctionInfoTermToText(FBPTerminal* Term, FBlueprintCompiledStatement* TargetLabel)
 {
-	check(LatentInfoStruct);
+	auto LatentInfoStruct = FLatentActionInfo::StaticStruct();
 
 	// Find the term name we need to fixup
 	FString FixupTermName;
@@ -339,7 +343,12 @@ void FBlueprintCompilerCppBackend::EmitCallDelegateStatment(FKismetFunctionConte
 
 	Emit(Body, TEXT("\t\t\t"));
 	Emit(Body, *SafeContextScope.GetAdditionalIndent());
-	Emit(Body, *FString::Printf(TEXT("%s.Broadcast("), *TermToText(Statement.FunctionContext)));
+	Emit(Body, *FString::Printf(TEXT("%s.Broadcast(%s)"), *TermToText(Statement.FunctionContext), *EmitMethodInputParameterList(Statement)));
+}
+
+FString FBlueprintCompilerCppBackend::EmitMethodInputParameterList(FBlueprintCompiledStatement& Statement)
+{
+	FString Result;
 	int32 NumParams = 0;
 	for (TFieldIterator<UProperty> PropIt(Statement.FunctionToCall); PropIt && (PropIt->PropertyFlags & CPF_Parm); ++PropIt)
 	{
@@ -349,7 +358,7 @@ void FBlueprintCompilerCppBackend::EmitCallDelegateStatment(FKismetFunctionConte
 		{
 			if (NumParams > 0)
 			{
-				Emit(Body, TEXT(", "));
+				Result += TEXT(", ");
 			}
 
 			FString VarName;
@@ -362,7 +371,7 @@ void FBlueprintCompilerCppBackend::EmitCallDelegateStatment(FKismetFunctionConte
 				// The target label will only ever be set on a call function when calling into the Ubergraph or
 				// on a latent function that will later call into the ubergraph, either of which requires a patchup
 				UStructProperty* StructProp = Cast<UStructProperty>(FuncParamProperty);
-				if (StructProp && StructProp->Struct == LatentInfoStruct)
+				if (StructProp && (StructProp->Struct == FLatentActionInfo::StaticStruct()))
 				{
 					// Latent function info case
 					VarName = LatentFunctionInfoTermToText(Term, Statement.TargetLabel);
@@ -381,118 +390,111 @@ void FBlueprintCompilerCppBackend::EmitCallDelegateStatment(FKismetFunctionConte
 
 			if (FuncParamProperty->HasAnyPropertyFlags(CPF_OutParm))
 			{
-				Emit(Body, TEXT("/*out*/ "));
+				Result += TEXT("/*out*/ ");
 			}
-			Emit(Body, *VarName);
+			Result += *VarName;
 
 			NumParams++;
 		}
 	}
-	Emit(Body, TEXT(");\n"));
+
+	return Result;
 }
 
-void FBlueprintCompilerCppBackend::EmitCallStatment(FKismetFunctionContext& FunctionContext, FBlueprintCompiledStatement& Statement)
+FString FBlueprintCompilerCppBackend::EmitCallStatmentInner(FBlueprintCompiledStatement& Statement, bool bInline)
 {
-	const bool bCallOnDifferentObject = (Statement.FunctionContext != NULL) && (Statement.FunctionContext->Name != TEXT("self"));
-	const bool bStaticCall = Statement.FunctionToCall->HasAnyFunctionFlags(FUNC_Static);
+	FString Result;
 
+	const bool bCallOnDifferentObject = Statement.FunctionContext && (Statement.FunctionContext->Name != TEXT("self"));
+	const bool bStaticCall = Statement.FunctionToCall->HasAnyFunctionFlags(FUNC_Static);
 	const bool bUseSafeContext = bCallOnDifferentObject && !bStaticCall;
-	FSafeContextScopedEmmitter SafeContextScope(Body, bUseSafeContext ? Statement.FunctionContext : NULL, *this, TEXT("\t\t\t"));
+	FSafeContextScopedEmmitter SafeContextScope(Result, bUseSafeContext ? Statement.FunctionContext : nullptr, *this, TEXT("\t\t\t"));
+	ensure(!bInline || !SafeContextScope.IsSafeContextUsed());
 
-	Emit(Body, TEXT("\t\t\t"));
-	Emit(Body, *SafeContextScope.GetAdditionalIndent());
-
-	// Handle the return value of the function being called
-	UProperty* FuncToCallReturnProperty = Statement.FunctionToCall->GetReturnProperty();
-	if (FuncToCallReturnProperty != NULL)
+	if (!bInline)
 	{
-		FBPTerminal* Term = Statement.LHS;
-		ensure(Term != NULL);
+		Result += TEXT("\t\t\t");
+		Result += SafeContextScope.GetAdditionalIndent();
 
-		FString ReturnVarName = TermToText(Term);
-		Emit(Body, *FString::Printf(TEXT("%s = "), *ReturnVarName));
+		// Handle the return value of the function being called
+		UProperty* FuncToCallReturnProperty = Statement.FunctionToCall->GetReturnProperty();
+		if (FuncToCallReturnProperty && ensure(Statement.LHS))
+		{
+			Result += FString::Printf(TEXT("%s = "), *TermToText(Statement.LHS));
+		}
 	}
 
 	// Emit object to call the method on
-	if (bCallOnDifferentObject) //@TODO: Badness, could be a self reference wired to another instance!
+	if (bStaticCall)
 	{
-		if (bStaticCall)
-		{
-			const bool bIsCustomThunk = Statement.FunctionToCall->HasMetaData(TEXT("CustomStructureParam")) || Statement.FunctionToCall->HasMetaData(TEXT("ArrayParm"));
-			FString FullFunctionNamePrefix = bIsCustomThunk
-				? TEXT("FCustomThunkTemplates::")
-				: FString::Printf(TEXT("U%s::"), *Statement.FunctionToCall->GetOuter()->GetName());
-			Emit(Body, *FullFunctionNamePrefix); //@TODO: Assuming U prefix but could be A
-		}
-		else
-		{
-			Emit(Body, *FString::Printf(TEXT("%s->"), *TermToText(Statement.FunctionContext, (UProperty*)(GetDefault<UObjectProperty>()))));
-		}
+		const bool bIsCustomThunk = Statement.FunctionToCall->HasMetaData(TEXT("CustomStructureParam")) || Statement.FunctionToCall->HasMetaData(TEXT("ArrayParm"));
+		auto OwnerClass = Statement.FunctionToCall->GetOuterUClass();
+		Result += bIsCustomThunk ? TEXT("FCustomThunkTemplates::") : FString::Printf(TEXT("%s%s::"), OwnerClass->GetPrefixCPP(), *OwnerClass->GetName());
 	}
-
-	// Emit method name
-	FString FunctionNameToCall;
-	Statement.FunctionToCall->GetName(FunctionNameToCall);
-	if (Statement.bIsParentContext)
+	else if(bCallOnDifferentObject) //@TODO: Badness, could be a self reference wired to another instance!
 	{
-		FunctionNameToCall = TEXT("Super::") + FunctionNameToCall;
+		Result += FString::Printf(TEXT("%s->"), *TermToText(Statement.FunctionContext, (UProperty*)(GetDefault<UObjectProperty>())));
 	}
-	Emit(Body, *FString::Printf(TEXT("%s"), *FunctionNameToCall));
 
-	// Emit method parameter list
-	Emit(Body, TEXT("("));
-	{
-		int32 NumParams = 0;
+	// Emit method name
+	Result += FString::Printf(TEXT("%s%s"), Statement.bIsParentContext ? TEXT("Super::") : TEXT(""), *Statement.FunctionToCall->GetName());
 
-		for (TFieldIterator<UProperty> PropIt(Statement.FunctionToCall); PropIt && (PropIt->PropertyFlags & CPF_Parm); ++PropIt)
-		{
-			UProperty* FuncParamProperty = *PropIt;
+	// Emit method parameter list
+	Result += TEXT("(");
+	Result += EmitMethodInputParameterList(Statement);
+	Result += TEXT(")");
 
-			if (!FuncParamProperty->HasAnyPropertyFlags(CPF_ReturnParm))
-			{
-				if (NumParams > 0)
-				{
-					Emit(Body, TEXT(", "));
-				}
+	if (!bInline)
+	{
+		Result += TEXT(";\n");
+	}
 
-				FString VarName;
+	return Result;
+}
 
-				FBPTerminal* Term = Statement.RHS[NumParams];
-				ensure(Term != NULL);
+void FBlueprintCompilerCppBackend::EmitCallStatment(FKismetFunctionContext& FunctionContext, FBlueprintCompiledStatement& Statement)
+{
+	Emit(Body, *EmitCallStatmentInner(Statement, false));
+}
 
-				if ((Statement.TargetLabel != NULL) && (Statement.UbergraphCallIndex == NumParams))
-				{
-					// The target label will only ever be set on a call function when calling into the Ubergraph or
-					// on a latent function that will later call into the ubergraph, either of which requires a patchup
-					UStructProperty* StructProp = Cast<UStructProperty>(FuncParamProperty);
-					if (StructProp && StructProp->Struct == LatentInfoStruct)
-					{
-						// Latent function info case
-						VarName = LatentFunctionInfoTermToText(Term, Statement.TargetLabel);
-					}
-					else
-					{
-						// Ubergraph entry point case
-						VarName = FString::FromInt(StateMapPerFunction[0].StatementToStateIndex(Statement.TargetLabel));
-					}
-				}
-				else
-				{
-					// Emit a normal parameter term
-					VarName = TermToText(Term, FuncParamProperty);
-				}
+FString FBlueprintCompilerCppBackend::EmitSwitchValueStatmentInner(FBlueprintCompiledStatement& Statement)
+{
+	check(Statement.RHS.Num() >= 2);
+	const int32 TermsBeforeCases = 1;
+	const int32 TermsPerCase = 2;
+	const int32 NumCases = ((Statement.RHS.Num() - 2) / TermsPerCase);
+	auto IndexTerm = Statement.RHS[0];
+	auto DefaultValueTerm = Statement.RHS.Last();
+
+	FString Result = FString::Printf(TEXT("TSwitchValue(%s, %s, %d")
+		, *TermToText(IndexTerm, nullptr) //index
+		, *TermToText(DefaultValueTerm, nullptr) // default
+		, NumCases);
+	
+	UProperty* VirtualIndexProperty = IndexTerm ? IndexTerm->AssociatedVarProperty : nullptr;
+	check(VirtualIndexProperty);
+	FStringOutputDevice IndexDeclaration;
+	VirtualIndexProperty->ExportCppDeclaration(IndexDeclaration, EExportedDeclaration::Parameter, NULL
+		, EPropertyExportCPPFlags::CPPF_CustomTypeName | EPropertyExportCPPFlags::CPPF_NoConst | EPropertyExportCPPFlags::CPPF_NoRef, true);
+
+	UProperty* VirtualValueProperty = DefaultValueTerm ? DefaultValueTerm->AssociatedVarProperty : nullptr;
+	check(VirtualValueProperty);
+	FStringOutputDevice ValueDeclaration;
+	VirtualValueProperty->ExportCppDeclaration(ValueDeclaration, EExportedDeclaration::Parameter, NULL
+		, EPropertyExportCPPFlags::CPPF_CustomTypeName | EPropertyExportCPPFlags::CPPF_NoConst | EPropertyExportCPPFlags::CPPF_NoRef, true);
+
+	for (int32 TermIndex = TermsBeforeCases; TermIndex < (NumCases * TermsPerCase); TermIndex += TermsPerCase)
+	{
+		Result += FString::Printf(TEXT(", TSwitchPair<%s, %s>(%s, %s)")
+			, *IndexDeclaration
+			, *ValueDeclaration
+			, *TermToText(Statement.RHS[TermIndex], VirtualIndexProperty)
+			, *TermToText(Statement.RHS[TermIndex + 1], VirtualValueProperty));
+	}
 
-				if (FuncParamProperty->HasAnyPropertyFlags(CPF_OutParm))
-				{
-					Emit(Body, TEXT("/*out*/ "));
-				}
-				Emit(Body, *VarName);
+	Result += TEXT(")");
 
-				NumParams++;
-			}
-		}
-	}
-	Emit(Body, TEXT(");\n"));
+	return Result;
 }
 
 void FBlueprintCompilerCppBackend::EmitAssignmentStatment(FKismetFunctionContext& FunctionContext, FBlueprintCompiledStatement& Statement)
@@ -953,29 +955,33 @@ void FBlueprintCompilerCppBackend::ConstructFunction(FKismetFunctionContext& Fun
 					case KCST_Nop:
 						Emit(Body, TEXT("\t\t\t//No operation.\n"));
 						break;
-					case KCST_WireTraceSite:
-						Emit(Body, TEXT("\t\t\t// Wire debug site.\n"));
-						break;
-					case KCST_DebugSite:
-						Emit(Body, TEXT("\t\t\t// Debug site.\n"));
-						break;
 					case KCST_CallFunction:
 						EmitCallStatment(FunctionContext, Statement);
 						break;
-					case KCST_CallDelegate:
-						EmitCallDelegateStatment(FunctionContext, Statement);
-						break;
 					case KCST_Assignment:
 						EmitAssignmentStatment(FunctionContext, Statement);
 						break;
-					case KCST_CastObjToInterface:
-						EmitCastObjToInterfaceStatement(FunctionContext, Statement);
+					case KCST_CompileError:
+						UE_LOG(LogK2Compiler, Error, TEXT("C++ backend encountered KCST_CompileError"));
+						Emit(Body, TEXT("\t\t\tstatic_assert(false); // KCST_CompileError"));
 						break;
-					case KCST_CrossInterfaceCast:
-						EmitCastBetweenInterfacesStatement(FunctionContext, Statement);
+					case KCST_PushState:
+						EmitPushStateStatement(FunctionContext, Statement);
 						break;
-					case KCST_CastInterfaceToObj:
-						EmitCastInterfaceToObjStatement(FunctionContext, Statement);
+					case KCST_Return:
+						Emit(Body, TEXT("\t\t\t// Return statement.\n"));
+						break;
+					case KCST_EndOfThread:
+						EmitEndOfThreadStatement(FunctionContext, ReturnValueString);
+						break;
+					case KCST_Comment:
+						Emit(Body, *FString::Printf(TEXT("\t\t\t// %s\n"), *Statement.Comment));
+						break;
+					case KCST_DebugSite:
+						Emit(Body, TEXT("\t\t\t// Debug site.\n"));
+						break;
+					case KCST_CastObjToInterface:
+						EmitCastObjToInterfaceStatement(FunctionContext, Statement);
 						break;
 					case KCST_DynamicCast:
 						EmitDynamicCastStatement(FunctionContext, Statement);
@@ -986,18 +992,33 @@ void FBlueprintCompilerCppBackend::ConstructFunction(FKismetFunctionContext& Fun
 					case KCST_AddMulticastDelegate:
 						EmitAddMulticastDelegateStatement(FunctionContext, Statement);
 						break;
-					case KCST_RemoveMulticastDelegate:
-						EmitRemoveMulticastDelegateStatement(FunctionContext, Statement);
+					case KCST_ClearMulticastDelegate:
+						EmitClearMulticastDelegateStatement(FunctionContext, Statement);
+						break;
+					case KCST_WireTraceSite:
+						Emit(Body, TEXT("\t\t\t// Wire debug site.\n"));
 						break;
 					case KCST_BindDelegate:
 						EmitBindDelegateStatement(FunctionContext, Statement);
 						break;
-					case KCST_ClearMulticastDelegate:
-						EmitClearMulticastDelegateStatement(FunctionContext, Statement);
+					case KCST_RemoveMulticastDelegate:
+						EmitRemoveMulticastDelegateStatement(FunctionContext, Statement);
+						break;
+					case KCST_CallDelegate:
+						EmitCallDelegateStatment(FunctionContext, Statement);
 						break;
 					case KCST_CreateArray:
 						EmitCreateArrayStatement(FunctionContext, Statement);
 						break;
+					case KCST_CrossInterfaceCast:
+						EmitCastBetweenInterfacesStatement(FunctionContext, Statement);
+						break;
+					case KCST_MetaCast:
+						EmitMetaCastStatement(FunctionContext, Statement);
+						break;
+					case KCST_CastInterfaceToObj:
+						EmitCastInterfaceToObjStatement(FunctionContext, Statement);
+						break;
 					case KCST_ComputedGoto:
 					case KCST_UnconditionalGoto:
 					case KCST_GotoIfNot:
@@ -1006,24 +1027,12 @@ void FBlueprintCompilerCppBackend::ConstructFunction(FKismetFunctionContext& Fun
 					case KCST_GotoReturnIfNot:
 						EmitGotoStatement(FunctionContext, Statement);
 						break;
-					case KCST_PushState:
-						EmitPushStateStatement(FunctionContext, Statement);
-						break;
-					case KCST_EndOfThread:
-						EmitEndOfThreadStatement(FunctionContext, ReturnValueString);
-						break;
-					case KCST_Comment:
-						Emit(Body, *FString::Printf(TEXT("\t\t\t// %s\n"), *Statement.Comment));
-						break;
-					case KCST_MetaCast:
-						EmitMetaCastStatement(FunctionContext, Statement);
-						break;
-					case KCST_Return:
-						Emit(Body, TEXT("\t\t\t// Return statement.\n"));
-						break;
+					case KCST_SwitchValue:
+						// Switch Value should be always an "inline" statement, so there is no point to handle it here
+					// case: KCST_AssignmentOnPersistentFrame
 					default:
 						Emit(Body, TEXT("\t// Warning: Ignoring unsupported statement\n"));
-						UE_LOG(LogK2Compiler, Warning, TEXT("C++ backend encountered unsupported statement type %d"), (int32)Statement.Type);
+						UE_LOG(LogK2Compiler, Error, TEXT("C++ backend encountered unsupported statement type %d"), (int32)Statement.Type);
 						break;
 					};
 				}
diff --git a/Engine/Source/Developer/BlueprintCompilerCppBackend/Private/BlueprintCompilerCppBackend.h b/Engine/Source/Developer/BlueprintCompilerCppBackend/Private/BlueprintCompilerCppBackend.h
--- a/Engine/Source/Developer/BlueprintCompilerCppBackend/Private/BlueprintCompilerCppBackend.h
+++ b/Engine/Source/Developer/BlueprintCompilerCppBackend/Private/BlueprintCompilerCppBackend.h
@@ -36,40 +36,19 @@ protected:
 
 	FString CppClassName;
 
-	// Pointers to commonly used structures (found in constructor)
-	UScriptStruct* VectorStruct;
-	UScriptStruct* RotatorStruct;
-	UScriptStruct* TransformStruct;
-	UScriptStruct* LatentInfoStruct;
-	UScriptStruct* LinearColorStruct;
 public:
 	FStringOutputDevice Header;
 	FStringOutputDevice Body;
-	FString TermToText(const FBPTerminal* Term, const UProperty* SourceProperty = NULL);
 
 	const FString& GetBody()	const override { return Body; }
 	const FString& GetHeader()	const override { return Header; }
 
-protected:
-	FString LatentFunctionInfoTermToText(FBPTerminal* Term, FBlueprintCompiledStatement* TargetLabel);
-
-	int32 StatementToStateIndex(FKismetFunctionContext& FunctionContext, FBlueprintCompiledStatement* Statement)
-	{
-		int32 Index = FunctionIndexMap.FindChecked(&FunctionContext);
-		return StateMapPerFunction[Index].StatementToStateIndex(Statement);
-	}
 public:
 
 	FBlueprintCompilerCppBackend(FKismetCompilerContext& InContext)
 		: MessageLog(InContext.MessageLog)
 		, CompilerContext(InContext)
 	{
-		extern UScriptStruct* Z_Construct_UScriptStruct_FVector();
-		VectorStruct = Z_Construct_UScriptStruct_FVector();
-		RotatorStruct = TBaseStructure<FRotator>::Get();
-		TransformStruct = TBaseStructure<FTransform>::Get();
-		LinearColorStruct = TBaseStructure<FLinearColor>::Get();
-		LatentInfoStruct = FLatentActionInfo::StaticStruct();
 	}
 
 	void Emit(FStringOutputDevice& Target, const TCHAR* Message)
@@ -109,4 +88,22 @@ public:
 
 	/** Builds both the header declaration and body implementation of a function */
 	void ConstructFunction(FKismetFunctionContext& FunctionContext, bool bGenerateStubOnly);
+
+	FString TermToText(const FBPTerminal* Term, const UProperty* SourceProperty = nullptr);
+
+protected:
+
+	FString LatentFunctionInfoTermToText(FBPTerminal* Term, FBlueprintCompiledStatement* TargetLabel);
+
+	int32 StatementToStateIndex(FKismetFunctionContext& FunctionContext, FBlueprintCompiledStatement* Statement)
+	{
+		int32 Index = FunctionIndexMap.FindChecked(&FunctionContext);
+		return StateMapPerFunction[Index].StatementToStateIndex(Statement);
+	}
+
+	FString EmitMethodInputParameterList(FBlueprintCompiledStatement& Statement);
+
+	FString EmitSwitchValueStatmentInner(FBlueprintCompiledStatement& Statement);
+
+	FString EmitCallStatmentInner(FBlueprintCompiledStatement& Statement, bool bInline);
 };
\ No newline at end of file
diff --git a/Engine/Source/Developer/BlueprintCompilerCppBackend/Private/BlueprintCompilerCppBackendUtils.h b/Engine/Source/Developer/BlueprintCompilerCppBackend/Private/BlueprintCompilerCppBackendUtils.h
--- a/Engine/Source/Developer/BlueprintCompilerCppBackend/Private/BlueprintCompilerCppBackendUtils.h
+++ b/Engine/Source/Developer/BlueprintCompilerCppBackend/Private/BlueprintCompilerCppBackendUtils.h
@@ -4,7 +4,7 @@
 struct FSafeContextScopedEmmitter
 {
 private:
-	FStringOutputDevice& Body;
+	FString& Body;
 	bool bSafeContextUsed;
 	const TCHAR* CurrentIndent;
 public:
@@ -18,7 +18,7 @@ public:
 		return bSafeContextUsed;
 	}
 
-	FSafeContextScopedEmmitter(FStringOutputDevice& InBody, const FBPTerminal* Term, FBlueprintCompilerCppBackend& CppBackend, const TCHAR* InCurrentIndent)
+	FSafeContextScopedEmmitter(FString& InBody, const FBPTerminal* Term, FBlueprintCompilerCppBackend& CppBackend, const TCHAR* InCurrentIndent)
 		: Body(InBody), bSafeContextUsed(false), CurrentIndent(InCurrentIndent)
 	{
 		TArray<FString> SafetyConditions;
@@ -127,7 +127,7 @@ struct FEmitHelper
 				if (!Pair.Value.IsEmpty())
 				{
 					FString Value = Pair.Value.Replace(TEXT("\n"), TEXT(""));
-					MetaDataStrings.Emplace(FString::Printf(TEXT("%s=\"%s\""), *Pair.Key.ToString(), *Pair.Value));
+					MetaDataStrings.Emplace(FString::Printf(TEXT("%s=\"%s\""), *Pair.Key.ToString(), *Value));
 				}
 				else
 				{
diff --git a/Engine/Source/Runtime/CoreUObject/Private/UObject/Property.cpp b/Engine/Source/Runtime/CoreUObject/Private/UObject/Property.cpp
--- a/Engine/Source/Runtime/CoreUObject/Private/UObject/Property.cpp
+++ b/Engine/Source/Runtime/CoreUObject/Private/UObject/Property.cpp
@@ -462,7 +462,8 @@ void UProperty::ExportCppDeclaration(FOutputDevice& Out, EExportedDeclaration::T
 	const uint32 ExportCPPFlags = AdditionalExportCPPFlags | (bIsParameter ? CPPF_ArgumentOrReturnValue : 0);
 	FString TypeText = GetCPPType(&ExtendedTypeText, ExportCPPFlags);
 
-	if (!dynamic_cast<const UBoolProperty*>(this)) // can't have const bitfields because then we cannot determine their offset and mask from the compiler
+	const bool bCanHaveConst = 0 == (AdditionalExportCPPFlags & CPPF_NoConst);
+	if (!dynamic_cast<const UBoolProperty*>(this) && bCanHaveConst) // can't have const bitfields because then we cannot determine their offset and mask from the compiler
 	{
 		const UObjectProperty* ObjectProp = dynamic_cast<const UObjectProperty*>(this);
 
@@ -503,6 +504,7 @@ void UProperty::ExportCppDeclaration(FOutputDevice& Out, EExportedDeclaration::T
 		}
 	}
 
+	const bool bCanHaveRef = 0 == (AdditionalExportCPPFlags & CPPF_NoConst);
 	if(auto BoolProperty = dynamic_cast<const UBoolProperty*>(this) )
 	{
 		// if this is a member variable, export it as a bitfield
@@ -517,7 +519,9 @@ void UProperty::ExportCppDeclaration(FOutputDevice& Out, EExportedDeclaration::T
 		else if( bIsParameter && HasAnyPropertyFlags(CPF_OutParm) )
 		{
 			// export as a reference
-			Out.Logf(TEXT("%s%s%s %s%s"), *TypeText, *ExtendedTypeText, TEXT("&"), *NameCpp, ArrayStr);
+			Out.Logf(TEXT("%s%s%s %s%s"), *TypeText, *ExtendedTypeText
+				, bCanHaveRef ? TEXT("&") : TEXT("")
+				, *NameCpp, ArrayStr);
 		}
 
 		else
@@ -542,16 +546,16 @@ void UProperty::ExportCppDeclaration(FOutputDevice& Out, EExportedDeclaration::T
 				{
 					// export as a reference (const ref if it isn't an out parameter)
 					Out.Logf(TEXT("%s%s%s%s %s"),
-						!HasAnyPropertyFlags(CPF_OutParm|CPF_ConstParm) ? TEXT("const ") : TEXT(""),
+						(bCanHaveConst && !HasAnyPropertyFlags(CPF_OutParm | CPF_ConstParm)) ? TEXT("const ") : TEXT(""),
 						*TypeText, *ExtendedTypeText,
-						TEXT("&"),
+						bCanHaveRef ? TEXT("&") : TEXT(""),
 						*NameCpp);
 				}
 				else
 				{
 					// export as a pointer if this is an optional out parm, reference if it's just an out parm, standard otherwise...
 					TCHAR ModifierString[2]={0,0};
-					if (HasAnyPropertyFlags(CPF_OutParm | CPF_ReferenceParm) || bIsInterfaceProp)
+					if (bCanHaveRef && (HasAnyPropertyFlags(CPF_OutParm | CPF_ReferenceParm) || bIsInterfaceProp))
 					{
 						ModifierString[0] = TEXT('&');
 					}
diff --git a/Engine/Source/Runtime/CoreUObject/Public/UObject/UnrealType.h b/Engine/Source/Runtime/CoreUObject/Public/UObject/UnrealType.h
--- a/Engine/Source/Runtime/CoreUObject/Public/UObject/UnrealType.h
+++ b/Engine/Source/Runtime/CoreUObject/Public/UObject/UnrealType.h
@@ -24,9 +24,13 @@ enum EPropertyExportCPPFlags
 	/** Indicates that we are exporting this property's CPP text for an argument or return value */
 	CPPF_ArgumentOrReturnValue		=	0x00000002,
 	/** Indicates thet we are exporting this property's CPP text for C++ definition of a function. */
-	CPPF_Implementation = 0x00000004,
+	CPPF_Implementation				=	0x00000004,
 	/** Indicates thet we are exporting this property's CPP text with an custom type name */
 	CPPF_CustomTypeName				=	0x00000008,
+	/** No 'const' keyword */
+	CPPF_NoConst					=	0x00000010,
+	/** No reference '&' sign */
+	CPPF_NoRef						=	0x00000020,
 };
 
 namespace EExportedDeclaration
diff --git a/Engine/Source/Runtime/Engine/Public/GeneratedCodeHelpers.h b/Engine/Source/Runtime/Engine/Public/GeneratedCodeHelpers.h
--- a/Engine/Source/Runtime/Engine/Public/GeneratedCodeHelpers.h
+++ b/Engine/Source/Runtime/Engine/Public/GeneratedCodeHelpers.h
@@ -111,9 +111,8 @@ struct FCustomThunkTemplates
 	template<typename T>
 	static void Array_Set(TArray<T>& TargetArray, int32 Index, const T& Item, bool bSizeToFit)
 	{
-		if (ensure(Index > 0))
+		if (ensure(Index >= 0))
 		{
-			const bool bIsValidIndex = ;
 			if (!TargetArray.IsValidIndex(Index) && bSizeToFit)
 			{
 				TargetArray.SetNum(Index + 1);
@@ -153,4 +152,57 @@ struct FCustomThunkTemplates
 	{
 		return UKismetSystemLibrary::Generic_SetStructurePropertyByName(Object, PropertyName, &Value);
 	}
-};
\ No newline at end of file
+};
+
+template<typename IndexType, typename ValueType>
+struct TSwitchPair
+{
+	const IndexType& IndexRef;
+	ValueType& ValueRef;
+
+	TSwitchPair(const IndexType& InIndexRef, ValueType& InValueRef)
+		: IndexRef(InIndexRef)
+		, ValueRef(InValueRef)
+	{}
+};
+
+#if PLATFORM_COMPILER_HAS_VARIADIC_TEMPLATES
+template<typename IndexType, typename ValueType>
+ValueType& TSwitchValue(const IndexType& CurrentIndex, ValueType& DefaultValue, int OptionsNum)
+{
+	return DefaultValue;
+}
+
+template<typename IndexType, typename ValueType, typename Head, typename... Tail>
+ValueType& TSwitchValue(const IndexType& CurrentIndex, ValueType& DefaultValue, int OptionsNum, Head HeadOption, Tail... TailOptions)
+{
+	if (CurrentIndex == HeadOption.IndexRef)
+	{
+		return HeadOption.ValueRef;
+	}
+	return TSwitchValue<IndexType, ValueType, Tail...>(CurrentIndex, DefaultValue, OptionsNum, TailOptions...);
+}
+#else //PLATFORM_COMPILER_HAS_VARIADIC_TEMPLATES
+template<typename IndexType, typename ValueType>
+ValueType& TSwitchValue(const IndexType& CurrentIndex, ValueType& DefaultValue, int OptionsNum, ...)
+{
+	typedef TSwitchPair < IndexType, ValueType > OptionType;
+
+	ValueType* SelectedValuePtr = nullptr;
+
+	va_list Options;
+	va_start(Options, OptionsNum);
+	for (int OptionIt = 0; OptionIt < OptionsNum; ++OptionIt)
+	{
+		OptionType Option = va_arg(Options, OptionType);
+		if (Option.IndexRef == CurrentIndex)
+		{
+			SelectedValuePtr = &Option.ValueRef;
+			break;
+		}
+	}
+	va_end(Options);
+
+	return SelectedValuePtr ? *SelectedValuePtr : DefaultValue;
+}
+#endif //PLATFORM_COMPILER_HAS_VARIADIC_TEMPLATES
\ No newline at end of file
