UEBP-40 C++ code generation (WIP):
Structs and Enums
Many improvements in BlueprintCompilerCppBackend

diff --git a/Engine/Source/Developer/BlueprintCompilerCppBackend/Private/BlueprintCompilerCppBackend.cpp b/Engine/Source/Developer/BlueprintCompilerCppBackend/Private/BlueprintCompilerCppBackend.cpp
--- a/Engine/Source/Developer/BlueprintCompilerCppBackend/Private/BlueprintCompilerCppBackend.cpp
+++ b/Engine/Source/Developer/BlueprintCompilerCppBackend/Private/BlueprintCompilerCppBackend.cpp
@@ -33,10 +33,10 @@ FString FBlueprintCompilerCppBackend::TermToText(const FBPTerminal* Term, const
 		else if (auto ByteProperty = Cast<const UByteProperty>(CoerceProperty))
 		{
 			// The PinSubCategoryObject check is to allow enum literals communicate with byte properties as literals
-			if (ByteProperty->Enum != NULL ||
-				(NULL != Cast<UEnum>(Term->Type.PinSubCategoryObject.Get())))
+			auto TypeEnum = ByteProperty->Enum ? ByteProperty->Enum : Cast<UEnum>(Term->Type.PinSubCategoryObject.Get());
+			if (TypeEnum)
 			{
-				return Term->Name;
+				return FString::Printf(TEXT("%s::%s"), *TypeEnum->GetName(), *Term->Name);
 			}
 			else
 			{
@@ -110,7 +110,11 @@ FString FBlueprintCompilerCppBackend::TermToText(const FBPTerminal* Term, const
 			}
 			else if (Term->ObjectLiteral)
 			{
-				if (auto LiteralClass = Cast<const UClass>(Term->ObjectLiteral))
+				if (auto AssetClassProp = Cast<UAssetObjectProperty>(CoerceProperty))
+				{
+					return FString::Printf(TEXT("FStringAssetReference(TEXT(\"%s\"))"), *(Term->ObjectLiteral->GetPathName()));
+				}
+				else if (auto LiteralClass = Cast<const UClass>(Term->ObjectLiteral))
 				{
 					return FString::Printf(TEXT("%s%s::StaticClass()"), LiteralClass->GetPrefixCPP(), *LiteralClass->GetName());
 				}
@@ -180,7 +184,7 @@ FString FBlueprintCompilerCppBackend::TermToText(const FBPTerminal* Term, const
 			}
 		}
 
-		Prefix += Term->Name;
+		Prefix += Term->AssociatedVarProperty ? Term->AssociatedVarProperty->GetNameCPP() : Term->Name;
 		return Prefix;
 	}
 }
@@ -214,7 +218,7 @@ FString FBlueprintCompilerCppBackend::LatentFunctionInfoTermToText(FBPTerminal*
 	return FString(TEXT("F")) + LatentInfoStruct->GetName() + StructValues;
 }
 
-void FBlueprintCompilerCppBackend::EmitClassProperties(FStringOutputDevice& Target, UClass* SourceClass)
+void FBlueprintCompilerCppBackend::EmitStructProperties(FStringOutputDevice& Target, UStruct* SourceClass)
 {
 	// Emit class variables
 	for (TFieldIterator<UProperty> It(SourceClass, EFieldIteratorFlags::ExcludeSuper); It; ++It)
@@ -240,15 +244,12 @@ void FBlueprintCompilerCppBackend::EmitClassProperties(FStringOutputDevice& Targ
 	}
 }
 
-void FBlueprintCompilerCppBackend::GenerateCodeFromClass(UClass* SourceClass, FString NewClassName, TIndirectArray<FKismetFunctionContext>& Functions, bool bGenerateStubsOnly)
+void FBlueprintCompilerCppBackend::GenerateCodeFromClass(UClass* SourceClass, TIndirectArray<FKismetFunctionContext>& Functions, bool bGenerateStubsOnly)
 {
-	auto CleanCppClassName = NewClassName.IsEmpty() ? SourceClass->GetName() : NewClassName;
+	auto CleanCppClassName = SourceClass->GetName();
 	CppClassName = FString(SourceClass->GetPrefixCPP()) + CleanCppClassName;
-
-	UClass* SuperClass = SourceClass->GetSuperClass();
-
-	Emit(Header, TEXT("#pragma once\n\n"));
-	Emit(Header, *FString::Printf(TEXT("#include \"%s.generated.h\"\n\n"), *CleanCppClassName));
+	
+	EmitFileBeginning(CleanCppClassName, SourceClass);
 
 	// MC DELEGATE DECLARATION
 	{
@@ -288,16 +289,17 @@ void FBlueprintCompilerCppBackend::GenerateCodeFromClass(UClass* SourceClass, FS
 		}
 	}
 
-	Emit(Header, TEXT("UCLASS(Blueprintable)\n"));
+	Emit(Header, TEXT("UCLASS(Blueprintable, BlueprintType)\n"));
 
+	UClass* SuperClass = SourceClass->GetSuperClass();
 	Emit(Header,
 		*FString::Printf(TEXT("class %s : public %s%s\n"), *CppClassName, SuperClass->GetPrefixCPP(), *SuperClass->GetName()));
 	Emit(Header,
 		TEXT("{\n")
 		TEXT("public:\n"));
-	Emit(Header, *FString::Printf(TEXT("\tGENERATED_UCLASS_BODY()\n")));
+	Emit(Header, *FString::Printf(TEXT("\tGENERATED_BODY()\n")));
 
-	EmitClassProperties(Header, SourceClass);
+	EmitStructProperties(Header, SourceClass);
 
 	// Create the state map
 	for (int32 i = 0; i < Functions.Num(); ++i)
@@ -318,10 +320,8 @@ void FBlueprintCompilerCppBackend::GenerateCodeFromClass(UClass* SourceClass, FS
 	PersistentHeaders.Add(TEXT("GeneratedCodeHelpers.h"));
 	Emit(Body, *FEmitHelper::GatherHeadersToInclude(SourceClass, PersistentHeaders));
 
-	//constructor
-	Emit(Body, *FString::Printf(
-		TEXT("%s::%s(const FObjectInitializer& ObjectInitializer) : Super(ObjectInitializer) {}\n\n"),
-		*CppClassName, *CppClassName));
+	//constructor is generated by GENERATED_BODY macto
+	// Emit(Body, *FString::Printf(TEXT("%s::%s(const FObjectInitializer& ObjectInitializer) : Super(ObjectInitializer) {}\n\n"), *CppClassName, *CppClassName));
 
 	for (int32 i = 0; i < Functions.Num(); ++i)
 	{
@@ -808,8 +808,8 @@ void FBlueprintCompilerCppBackend::ConstructFunction(FKismetFunctionContext& Fun
 					}
 					else
 					{
-						MessageLog.Error(*FString::Printf(TEXT("Function %s from graph @@ has more than one return value (named %s and %s)"),
-							*FunctionName, *ReturnValue->GetName(), *Property->GetName()), FunctionContext.SourceGraph);
+						UE_LOG(LogK2Compiler, Error, TEXT("Function %s from graph @@ has more than one return value (named %s and %s)"),
+							*FunctionName, *GetPathNameSafe(FunctionContext.SourceGraph), *ReturnValue->GetName(), *Property->GetName());
 					}
 				}
 				else
@@ -1029,7 +1029,7 @@ void FBlueprintCompilerCppBackend::ConstructFunction(FKismetFunctionContext& Fun
 						break;
 					case KCST_SwitchValue:
 						// Switch Value should be always an "inline" statement, so there is no point to handle it here
-					// case: KCST_AssignmentOnPersistentFrame
+						// case: KCST_AssignmentOnPersistentFrame
 					default:
 						Emit(Body, TEXT("\t// Warning: Ignoring unsupported statement\n"));
 						UE_LOG(LogK2Compiler, Error, TEXT("C++ backend encountered unsupported statement type %d"), (int32)Statement.Type);
@@ -1048,4 +1048,104 @@ void FBlueprintCompilerCppBackend::ConstructFunction(FKismetFunctionContext& Fun
 	EmitReturnStatement(FunctionContext, ReturnValueString);
 
 	Emit(Body, TEXT("}\n\n"));
-}
\ No newline at end of file
+}
+
+void FBlueprintCompilerCppBackend::GenerateCodeFromEnum(UUserDefinedEnum* SourceEnum)
+{
+	check(SourceEnum);
+	const FString Name = SourceEnum->GetName();
+	EmitFileBeginning(Name, nullptr);
+
+	Emit(Header, TEXT("UENUM(BlueprintType)\nenum class "));
+	Emit(Header, *Name);
+	Emit(Header, TEXT(" : uint8\n{"));
+
+	for (int32 Index = 0; Index < SourceEnum->NumEnums(); ++Index)
+	{
+		const FString ElemName = SourceEnum->GetEnumName(Index);
+		const int32 ElemValue = Index;
+
+		const FString& DisplayNameMD = SourceEnum->GetMetaData(TEXT("DisplayName"), ElemValue);// TODO: value or index?
+		const FString Meta = DisplayNameMD.IsEmpty() ? FString() : FString::Printf(TEXT("UMETA(DisplayName = \"%s\")"), *DisplayNameMD);
+		Emit(Header, *FString::Printf(TEXT("\n\t%s = %d %s,"), *ElemName, ElemValue, *Meta));
+	}
+	Emit(Header, TEXT("\n};\n"));
+}
+
+void FBlueprintCompilerCppBackend::GenerateCodeFromStruct(UUserDefinedStruct* SourceStruct)
+{
+	check(SourceStruct);
+
+	EmitFileBeginning(SourceStruct->GetName(), SourceStruct);
+
+	const FString NewName = FString(TEXT("F")) + SourceStruct->GetName();
+	Emit(Header, TEXT("USTRUCT(BlueprintType)\n"));
+
+	Emit(Header, *FString::Printf(TEXT("struct %s\n{\npublic:\n\tGENERATED_BODY()\n"), *NewName));
+	EmitStructProperties(Header, SourceStruct);
+	Emit(Header, TEXT("};\n"));
+}
+
+void FBlueprintCompilerCppBackend::EmitFileBeginning(const FString& CleanName, UStruct* SourceStruct)
+{
+	Emit(Header, TEXT("#pragma once\n\n"));
+	if (SourceStruct)
+	{
+		// find objects referenced by functions/script
+		TArray<UObject*> IncludeInHeader;
+		TArray<UObject*> IncludeInBody;
+		{
+			FReferenceFinder HeaderReferenceFinder(IncludeInHeader, NULL, false, false, false, false);
+			FReferenceFinder BodyReferenceFinder(IncludeInBody, NULL, false, false, false, false);
+			{
+				TArray<UObject*> ObjectsToCheck;
+				GetObjectsWithOuter(SourceStruct, ObjectsToCheck, true);
+				for (auto Obj : ObjectsToCheck)
+				{
+					auto Property = Cast<UProperty>(Obj);
+					auto OwnerProperty = IsValid(Property) ? Property->GetOwnerProperty() : nullptr;
+					const bool bIsParam = OwnerProperty && (0 != (OwnerProperty->PropertyFlags & CPF_Parm));
+					const bool bIsMemberVariable = OwnerProperty && (OwnerProperty->GetOuter() == SourceStruct);
+					if (bIsParam || bIsMemberVariable)
+					{
+						HeaderReferenceFinder.FindReferences(Obj);
+					}
+					else
+					{
+						BodyReferenceFinder.FindReferences(Obj);
+					}
+				}
+			}
+
+			if (auto SuperStruct = SourceStruct->GetSuperStruct())
+			{
+				IncludeInHeader.AddUnique(SuperStruct);
+			}
+		}
+
+		TSet<FString> AlreadyIncluded;
+		AlreadyIncluded.Add(SourceStruct->GetName());
+
+		auto EmitInner = [&](FStringOutputDevice& Dst, TArray<UObject*>& Src)
+		{
+			for (UObject* Obj : Src)
+			{
+				const bool bWantedType = Obj && (Obj->IsA<UBlueprintGeneratedClass>() || Obj->IsA<UUserDefinedEnum>() || Obj->IsA<UUserDefinedStruct>());
+				if (bWantedType)
+				{
+					const FString Name = Obj->GetName();
+					if (!AlreadyIncluded.Contains(Name))
+					{
+						AlreadyIncluded.Add(Name);
+						Emit(Dst, *FString::Printf(TEXT("#include \"%s.h\"\n"), *Name));
+					}
+				}
+			}
+			Emit(Dst, TEXT("\n"));
+		};
+
+		EmitInner(Header, IncludeInHeader);
+		EmitInner(Body, IncludeInBody);
+	}
+	Emit(Header, *FString::Printf(TEXT("#include \"%s.generated.h\"\n\n"), *CleanName));
+}
diff --git a/Engine/Source/Developer/BlueprintCompilerCppBackend/Private/BlueprintCompilerCppBackend.h b/Engine/Source/Developer/BlueprintCompilerCppBackend/Private/BlueprintCompilerCppBackend.h
--- a/Engine/Source/Developer/BlueprintCompilerCppBackend/Private/BlueprintCompilerCppBackend.h
+++ b/Engine/Source/Developer/BlueprintCompilerCppBackend/Private/BlueprintCompilerCppBackend.h
@@ -6,9 +6,6 @@
 class FBlueprintCompilerCppBackend : public IBlueprintCompilerCppBackend
 {
 protected:
-	FCompilerResultsLog& MessageLog;
-	FKismetCompilerContext& CompilerContext;
-
 	struct FFunctionLabelInfo
 	{
 		TMap<FBlueprintCompiledStatement*, int32> StateMap;
@@ -40,25 +37,23 @@ public:
 	FStringOutputDevice Header;
 	FStringOutputDevice Body;
 
-	const FString& GetBody()	const override { return Body; }
-	const FString& GetHeader()	const override { return Header; }
+	// IBlueprintCompilerCppBackend implementation
+	virtual void GenerateCodeFromClass(UClass* SourceClass, TIndirectArray<FKismetFunctionContext>& Functions, bool bGenerateStubsOnly) override;
+	virtual void GenerateCodeFromEnum(UUserDefinedEnum* SourceEnum) override;
+	virtual void GenerateCodeFromStruct(UUserDefinedStruct* SourceStruct) override;
 
-public:
+	virtual const FString& GetBody()	const override { return Body; }
+	virtual const FString& GetHeader()	const override { return Header; }
+	// end of IBlueprintCompilerCppBackend implementation
 
-	FBlueprintCompilerCppBackend(FKismetCompilerContext& InContext)
-		: MessageLog(InContext.MessageLog)
-		, CompilerContext(InContext)
-	{
-	}
+public:
 
 	void Emit(FStringOutputDevice& Target, const TCHAR* Message)
 	{
 		Target += Message;
 	}
 
-	void EmitClassProperties(FStringOutputDevice& Target, UClass* SourceClass);
-
-	void GenerateCodeFromClass(UClass* SourceClass, FString NewClassName, TIndirectArray<FKismetFunctionContext>& Functions, bool bGenerateStubsOnly) override;
+	void EmitStructProperties(FStringOutputDevice& Target, UStruct* SourceClass);
 
 	void EmitCallStatment(FKismetFunctionContext& FunctionContext, FBlueprintCompiledStatement& Statement);
 	void EmitCallDelegateStatment(FKismetFunctionContext& FunctionContext, FBlueprintCompiledStatement& Statement);
@@ -92,6 +87,7 @@ public:
 	FString TermToText(const FBPTerminal* Term, const UProperty* SourceProperty = nullptr);
 
 protected:
+	void EmitFileBeginning(const FString& CleanName, UStruct* SourceStruct);
 
 	FString LatentFunctionInfoTermToText(FBPTerminal* Term, FBlueprintCompiledStatement* TargetLabel);
 
diff --git a/Engine/Source/Developer/BlueprintCompilerCppBackend/Private/BlueprintCompilerCppBackendModule.cpp b/Engine/Source/Developer/BlueprintCompilerCppBackend/Private/BlueprintCompilerCppBackendModule.cpp
--- a/Engine/Source/Developer/BlueprintCompilerCppBackend/Private/BlueprintCompilerCppBackendModule.cpp
+++ b/Engine/Source/Developer/BlueprintCompilerCppBackend/Private/BlueprintCompilerCppBackendModule.cpp
@@ -10,7 +10,7 @@ class FBlueprintCompilerCppBackendModule : public IBlueprintCompilerCppBackendMo
 	/** IModuleInterface implementation */
 	virtual void StartupModule() override;
 	virtual void ShutdownModule() override;
-	virtual IBlueprintCompilerCppBackend* Create(FKismetCompilerContext& InContext)  override;
+	virtual IBlueprintCompilerCppBackend* Create()  override;
 };
 
 IMPLEMENT_MODULE(FBlueprintCompilerCppBackendModule, BlueprintCompilerCppBackend)
@@ -29,7 +29,7 @@ void FBlueprintCompilerCppBackendModule::ShutdownModule()
 	// we call this function before unloading the module.
 }
 
-IBlueprintCompilerCppBackend* FBlueprintCompilerCppBackendModule::Create(FKismetCompilerContext& InContext)
+IBlueprintCompilerCppBackend* FBlueprintCompilerCppBackendModule::Create()
 {
-	return new FBlueprintCompilerCppBackend(InContext);
+	return new FBlueprintCompilerCppBackend();
 }
diff --git a/Engine/Source/Editor/Kismet/Private/NativeCodeGenerationTool.cpp b/Engine/Source/Editor/Kismet/Private/NativeCodeGenerationTool.cpp
--- a/Engine/Source/Editor/Kismet/Private/NativeCodeGenerationTool.cpp
+++ b/Engine/Source/Editor/Kismet/Private/NativeCodeGenerationTool.cpp
@@ -8,6 +8,7 @@
 #include "Engine/UserDefinedStruct.h"
 #include "SourceCodeNavigation.h"
 #include "DesktopPlatformModule.h"
+//#include "Editor/KismetCompiler/Public/BlueprintCompilerCppBackendInterface.h"
 
 #define LOCTEXT_NAMESPACE "NativeCodeGenerationTool"
 
@@ -19,9 +20,7 @@
 struct FGeneratedCodeData
 {
 	FGeneratedCodeData(UBlueprint& InBlueprint) 
-		: HeaderSource(new FString())
-		, CppSource(new FString())
-		, Blueprint(&InBlueprint)
+		: Blueprint(&InBlueprint)
 	{
 		FName GeneratedClassName, SkeletonClassName;
 		InBlueprint.GetBlueprintClassNames(GeneratedClassName, SkeletonClassName);
@@ -31,11 +30,10 @@ struct FGeneratedCodeData
 	}
 
 	FString TypeDependencies;
-	TSharedPtr<FString> HeaderSource;
-	TSharedPtr<FString> CppSource;
 	FString ErrorString;
 	FString ClassName;
 	TWeakObjectPtr<UBlueprint> Blueprint;
+	TSet<UObject*> DependentObjects;
 
 	void GatherUserDefinedDependencies(UBlueprint& InBlueprint)
 	{
@@ -70,9 +68,11 @@ struct FGeneratedCodeData
 				{
 					TypeDependencies += Obj->GetPathName();
 					TypeDependencies += TEXT("\n");
+					DependentObjects.Add(Obj);
 				}
 			}
 		}
+		DependentObjects.Add(InBlueprint.GeneratedClass);
 
 		if (TypeDependencies.IsEmpty())
 		{
@@ -104,72 +104,80 @@ struct FGeneratedCodeData
 
 	bool Save(const FString& HeaderDirPath, const FString& CppDirPath)
 	{
-		FScopedSlowTask SlowTask(7, LOCTEXT("GeneratingCppFiles", "Generating C++ files.."));
-		SlowTask.MakeDialog();
-		SlowTask.EnterProgressFrame();
-
 		if (!Blueprint.IsValid())
 		{
 			ErrorString += LOCTEXT("InvalidBlueprint", "Invalid Blueprint\n").ToString();
 			return false;
 		}
 
-		FKismetEditorUtilities::GenerateCppCode(Blueprint.Get(), HeaderSource, CppSource, ClassName);
-
+		const int WorkParts = 4 + 2 * DependentObjects.Num();
+		FScopedSlowTask SlowTask(WorkParts, LOCTEXT("GeneratingCppFiles", "Generating C++ files.."));
+		SlowTask.MakeDialog();
 		SlowTask.EnterProgressFrame();
 
-		bool bProjectHadCodeFiles = false;
-		{
-			TArray<FString> OutProjectCodeFilenames;
-			IFileManager::Get().FindFilesRecursive(OutProjectCodeFilenames, *FPaths::GameSourceDir(), TEXT("*.h"), true, false, false);
-			IFileManager::Get().FindFilesRecursive(OutProjectCodeFilenames, *FPaths::GameSourceDir(), TEXT("*.cpp"), true, false, false);
-			bProjectHadCodeFiles = OutProjectCodeFilenames.Num() > 0;
-		}
-
 		TArray<FString> CreatedFiles;
-
-		const FString NeHeaderFilename = FPaths::Combine(*HeaderDirPath, *HeaderFileName());
-		const bool bHeaderSaved = FFileHelper::SaveStringToFile(*HeaderSource, *NeHeaderFilename);
-		if (!bHeaderSaved)
-		{
-			ErrorString += LOCTEXT("HeaderNotSaved", "Header file wasn't saved. Check log for details.\n").ToString();
-		}
-		else
-		{
-			CreatedFiles.Add(NeHeaderFilename);
-		}
-
-		SlowTask.EnterProgressFrame();
-
-		const FString NewCppFilename = FPaths::Combine(*CppDirPath, *SourceFileName());
-		const bool bCppSaved = FFileHelper::SaveStringToFile(*CppSource, *NewCppFilename);
-		if (!bCppSaved)
-		{
-			ErrorString += LOCTEXT("CppNotSaved", "Cpp file wasn't saved. Check log for details.\n").ToString();
-		}
-		else
+		for(auto Obj : DependentObjects)
 		{
-			CreatedFiles.Add(NewCppFilename);
-		}
+			TSharedPtr<FString> HeaderSource(new FString());
+			TSharedPtr<FString> CppSource(new FString());
+			FKismetEditorUtilities::GenerateCppCode(Obj, HeaderSource, CppSource);
+			SlowTask.EnterProgressFrame();
+
+			const FString FullHeaderFilename = FPaths::Combine(*HeaderDirPath, *(Obj->GetName() + TEXT(".h")));
+			const bool bHeaderSaved = FFileHelper::SaveStringToFile(*HeaderSource, *FullHeaderFilename);
+			if (!bHeaderSaved)
+			{
+				ErrorString += FString::Printf(*LOCTEXT("HeaderNotSaved", "Header file wasn't saved. Check log for details. %s\n").ToString(), *Obj->GetPathName());
+			}
+			else
+			{
+				CreatedFiles.Add(FullHeaderFilename);
+			}
 
-		TArray<FString> CreatedFilesForExternalAppRead;
-		CreatedFilesForExternalAppRead.Reserve(CreatedFiles.Num());
-		for (const FString& CreatedFile : CreatedFiles)
-		{
-			CreatedFilesForExternalAppRead.Add(IFileManager::Get().ConvertToAbsolutePathForExternalAppForRead(*CreatedFile));
+			SlowTask.EnterProgressFrame();
+			if (!CppSource->IsEmpty())
+			{
+				const FString NewCppFilename = FPaths::Combine(*CppDirPath, *(Obj->GetName() + TEXT(".cpp")));
+				const bool bCppSaved = FFileHelper::SaveStringToFile(*CppSource, *NewCppFilename);
+				if (!bCppSaved)
+				{
+					ErrorString += FString::Printf(*LOCTEXT("CppNotSaved", "Cpp file wasn't saved. Check log for details. %s\n").ToString(), *Obj->GetPathName());
+				}
+				else
+				{
+					CreatedFiles.Add(NewCppFilename);
+				}
+			}
 		}
 
 		SlowTask.EnterProgressFrame();
 
 		bool bGenerateProjectFiles = true;
-		// First see if we can avoid a full generation by adding the new files to an already open project
-		if (bProjectHadCodeFiles && FSourceCodeNavigation::AddSourceFiles(CreatedFilesForExternalAppRead))
 		{
-			// We successfully added the new files to the solution, but we still need to run UBT with -gather to update any UBT makefiles
-			if (FDesktopPlatformModule::Get()->InvalidateMakefiles(FPaths::RootDir(), FPaths::GetProjectFilePath(), GWarn))
+			bool bProjectHadCodeFiles = false;
 			{
-				// We managed the gather, so we can skip running the full generate
-				bGenerateProjectFiles = false;
+				TArray<FString> OutProjectCodeFilenames;
+				IFileManager::Get().FindFilesRecursive(OutProjectCodeFilenames, *FPaths::GameSourceDir(), TEXT("*.h"), true, false, false);
+				IFileManager::Get().FindFilesRecursive(OutProjectCodeFilenames, *FPaths::GameSourceDir(), TEXT("*.cpp"), true, false, false);
+				bProjectHadCodeFiles = OutProjectCodeFilenames.Num() > 0;
+			}
+
+			TArray<FString> CreatedFilesForExternalAppRead;
+			CreatedFilesForExternalAppRead.Reserve(CreatedFiles.Num());
+			for (const FString& CreatedFile : CreatedFiles)
+			{
+				CreatedFilesForExternalAppRead.Add(IFileManager::Get().ConvertToAbsolutePathForExternalAppForRead(*CreatedFile));
+			}
+
+			// First see if we can avoid a full generation by adding the new files to an already open project
+			if (bProjectHadCodeFiles && FSourceCodeNavigation::AddSourceFiles(CreatedFilesForExternalAppRead))
+			{
+				// We successfully added the new files to the solution, but we still need to run UBT with -gather to update any UBT makefiles
+				if (FDesktopPlatformModule::Get()->InvalidateMakefiles(FPaths::RootDir(), FPaths::GetProjectFilePath(), GWarn))
+				{
+					// We managed the gather, so we can skip running the full generate
+					bGenerateProjectFiles = false;
+				}
 			}
 		}
 
@@ -186,7 +194,7 @@ struct FGeneratedCodeData
 			}
 		}
 
-		return bHeaderSaved && bCppSaved && bProjectFileUpdated;
+		return ErrorString.IsEmpty();
 	}
 };
 
@@ -320,22 +328,6 @@ private:
 		return GeneratedCodeData.IsValid() ? FText::FromString(GeneratedCodeData->ClassName) : FText::GetEmpty();
 	}
 
-	void OnClassNameCommited(const FText& NameText, ETextCommit::Type TextCommit)
-	{
-		if (GeneratedCodeData.IsValid())
-		{
-			GeneratedCodeData->ClassName = NameText.ToString();
-			if (HeaderDirectoryBrowser.IsValid())
-			{
-				HeaderDirectoryBrowser->File = GeneratedCodeData->HeaderFileName();
-			}
-			if (SourceDirectoryBrowser.IsValid())
-			{
-				SourceDirectoryBrowser->File = GeneratedCodeData->SourceFileName();
-			}
-		}
-	}
-
 public:
 
 	void Construct(const FArguments& InArgs)
@@ -362,10 +354,8 @@ public:
 				.Padding(4.0f)
 				.AutoHeight()
 				[
-					SNew(SEditableTextBox)
+					SNew(STextBlock)
 					.Text(this, &SNativeCodeGenerationDialog::GetClassName)
-					.OnTextCommitted(this, &SNativeCodeGenerationDialog::OnClassNameCommited)
-					.IsEnabled(this, &SNativeCodeGenerationDialog::IsEditable)
 				]
 				+ SVerticalBox::Slot()
 				.Padding(4.0f)
diff --git a/Engine/Source/Editor/KismetCompiler/Private/KismetCompiler.cpp b/Engine/Source/Editor/KismetCompiler/Private/KismetCompiler.cpp
--- a/Engine/Source/Editor/KismetCompiler/Private/KismetCompiler.cpp
+++ b/Engine/Source/Editor/KismetCompiler/Private/KismetCompiler.cpp
@@ -3617,10 +3617,10 @@ void FKismetCompilerContext::Compile()
 		// Generate code thru the backend(s)
 		if ((bDisplayCpp && bIsFullCompile) || CompileOptions.DoesRequireCppCodeGeneration())
 		{
-			TUniquePtr<IBlueprintCompilerCppBackend> Backend_CPP(IBlueprintCompilerCppBackendModuleInterface::Get().Create(*this));
+			TUniquePtr<IBlueprintCompilerCppBackend> Backend_CPP(IBlueprintCompilerCppBackendModuleInterface::Get().Create());
 
 			// The C++ backend is currently only for debugging, so it's only run if the output will be visible
-			Backend_CPP->GenerateCodeFromClass(NewClass, CompileOptions.NewCppClassName, FunctionList, !bIsFullCompile);
+			Backend_CPP->GenerateCodeFromClass(NewClass, FunctionList, !bIsFullCompile);
 
 			if (CompileOptions.OutHeaderSourceCode.IsValid())
 			{
diff --git a/Engine/Source/Editor/KismetCompiler/Private/KismetCompilerModule.cpp b/Engine/Source/Editor/KismetCompiler/Private/KismetCompilerModule.cpp
--- a/Engine/Source/Editor/KismetCompiler/Private/KismetCompilerModule.cpp
+++ b/Engine/Source/Editor/KismetCompiler/Private/KismetCompilerModule.cpp
@@ -15,6 +15,7 @@
 
 #include "UserDefinedStructureCompilerUtils.h"
 #include "Engine/UserDefinedStruct.h"
+#include "BlueprintCompilerCppBackendInterface.h"
 
 DEFINE_LOG_CATEGORY(LogK2Compiler);
 DECLARE_CYCLE_STAT(TEXT("Compile Time"), EKismetCompilerStats_CompileTime, STATGROUP_KismetCompiler);
@@ -128,6 +129,20 @@ void FKismet2CompilerModule::CompileStructure(UUserDefinedStruct* Struct, FCompi
 	FUserDefinedStructureCompilerUtils::CompileStruct(Struct, Results, true);
 }
 
+FString FKismet2CompilerModule::GenerateCppCodeForEnum(UUserDefinedEnum* UDEnum)
+{
+	TUniquePtr<IBlueprintCompilerCppBackend> Backend_CPP(IBlueprintCompilerCppBackendModuleInterface::Get().Create());
+	Backend_CPP->GenerateCodeFromEnum(UDEnum);
+	return Backend_CPP->GetHeader();
+}
+
+FString FKismet2CompilerModule::GenerateCppCodeForStruct(UUserDefinedStruct* UDStruct)
+{
+	TUniquePtr<IBlueprintCompilerCppBackend> Backend_CPP(IBlueprintCompilerCppBackendModuleInterface::Get().Create());
+	Backend_CPP->GenerateCodeFromStruct(UDStruct);
+	return Backend_CPP->GetHeader();
+}
+
 extern UNREALED_API FSecondsCounterData BlueprintCompileAndLoadTimerData;
 
 // Compiles a blueprint.
diff --git a/Engine/Source/Editor/KismetCompiler/Public/BlueprintCompilerCppBackendInterface.h b/Engine/Source/Editor/KismetCompiler/Public/BlueprintCompilerCppBackendInterface.h
--- a/Engine/Source/Editor/KismetCompiler/Public/BlueprintCompilerCppBackendInterface.h
+++ b/Engine/Source/Editor/KismetCompiler/Public/BlueprintCompilerCppBackendInterface.h
@@ -8,7 +8,9 @@
 class IBlueprintCompilerCppBackend
 {
 public:
-	virtual void GenerateCodeFromClass(UClass* SourceClass, FString NewClassName, TIndirectArray<FKismetFunctionContext>& Functions, bool bGenerateStubsOnly) = 0;
+	virtual void GenerateCodeFromClass(UClass* SourceClass, TIndirectArray<FKismetFunctionContext>& Functions, bool bGenerateStubsOnly) = 0;
+	virtual void GenerateCodeFromEnum(UUserDefinedEnum* SourceEnum) = 0;
+	virtual void GenerateCodeFromStruct(UUserDefinedStruct* SourceStruct) = 0;
 
 	virtual const FString& GetBody() const = 0;
 	virtual const FString& GetHeader() const = 0;
@@ -44,5 +46,5 @@ public:
 		return FModuleManager::Get().IsModuleLoaded("BlueprintCompilerCppBackend");
 	}
 
-	virtual IBlueprintCompilerCppBackend* Create(FKismetCompilerContext& InContext) = 0;
+	virtual IBlueprintCompilerCppBackend* Create() = 0;
 };
\ No newline at end of file
diff --git a/Engine/Source/Editor/KismetCompiler/Public/KismetCompilerModule.h b/Engine/Source/Editor/KismetCompiler/Public/KismetCompilerModule.h
--- a/Engine/Source/Editor/KismetCompiler/Public/KismetCompilerModule.h
+++ b/Engine/Source/Editor/KismetCompiler/Public/KismetCompilerModule.h
@@ -84,6 +84,9 @@ public:
 	 * blueprint is a normal UBlueprint, like UUserWidget blueprints should be UWidgetBlueprints.
 	 */
 	virtual void GetBlueprintTypesForClass(UClass* ParentClass, UClass*& OutBlueprintClass, UClass*& OutBlueprintGeneratedClass) const = 0;
+
+	virtual FString GenerateCppCodeForEnum(UUserDefinedEnum* UDEnum) = 0;
+	virtual FString GenerateCppCodeForStruct(UUserDefinedStruct* UDStruct) = 0;
 };
 
 //////////////////////////////////////////////////////////////////////////
@@ -102,6 +105,8 @@ public:
 	virtual void RemoveBlueprintGeneratedClasses(class UBlueprint* Blueprint) override;
 	virtual TArray<IBlueprintCompiler*>& GetCompilers() override { return Compilers; }
 	virtual void GetBlueprintTypesForClass(UClass* ParentClass, UClass*& OutBlueprintClass, UClass*& OutBlueprintGeneratedClass) const override;
+	virtual FString GenerateCppCodeForEnum(UUserDefinedEnum* UDEnum) override;
+	virtual FString GenerateCppCodeForStruct(UUserDefinedStruct* UDStruct) override;
 	// End implementation
 private:
 	void CompileBlueprintInner(class UBlueprint* Blueprint, const FKismetCompilerOptions& CompileOptions, FCompilerResultsLog& Results, TSharedPtr<FBlueprintCompileReinstancer> Reinstancer, TArray<UObject*>* ObjLoaded);
diff --git a/Engine/Source/Editor/UnrealEd/Private/Kismet2/Kismet2.cpp b/Engine/Source/Editor/UnrealEd/Private/Kismet2/Kismet2.cpp
--- a/Engine/Source/Editor/UnrealEd/Private/Kismet2/Kismet2.cpp
+++ b/Engine/Source/Editor/UnrealEd/Private/Kismet2/Kismet2.cpp
@@ -942,35 +942,57 @@ void FKismetEditorUtilities::RecompileBlueprintBytecode(UBlueprint* BlueprintObj
 }
 
 /** Recompiles the bytecode of a blueprint only.  Should only be run for recompiling dependencies during compile on load */
-void FKismetEditorUtilities::GenerateCppCode(UBlueprint* InBlueprintObj, TSharedPtr<FString> OutHeaderSource, TSharedPtr<FString> OutCppSource, const FString& OptionalClassName)
+void FKismetEditorUtilities::GenerateCppCode(UObject* Obj, TSharedPtr<FString> OutHeaderSource, TSharedPtr<FString> OutCppSource)
 {
-	check(InBlueprintObj);
-	check(InBlueprintObj->GetOutermost() != GetTransientPackage());
-	checkf(InBlueprintObj->GeneratedClass, TEXT("Invalid generated class for %s"), *InBlueprintObj->GetName());
-	check(OutHeaderSource.IsValid());
-	check(OutCppSource.IsValid());
+	auto UDEnum = Cast<UUserDefinedEnum>(Obj);
+	auto UDStruct = Cast<UUserDefinedStruct>(Obj);
+	auto BPGC = Cast<UClass>(Obj);
+	auto InBlueprintObj = BPGC ? Cast<UBlueprint>(BPGC->ClassGeneratedBy) : nullptr;
 
-	//TGuardValue<bool> DuplicatingReadOnly(InBlueprintObj->bDuplicatingReadOnly, true);
+	if (InBlueprintObj)
 	{
-		auto BlueprintObj = DuplicateObject<UBlueprint>(InBlueprintObj, GetTransientPackage(), *InBlueprintObj->GetName());
+		check(InBlueprintObj->GetOutermost() != GetTransientPackage());
+		checkf(InBlueprintObj->GeneratedClass, TEXT("Invalid generated class for %s"), *InBlueprintObj->GetName());
+		check(OutHeaderSource.IsValid());
+		check(OutCppSource.IsValid());
+
+		//TGuardValue<bool> DuplicatingReadOnly(InBlueprintObj->bDuplicatingReadOnly, true);
 		{
-			auto Reinstancer = FBlueprintCompileReinstancer::Create(BlueprintObj->GeneratedClass);
+			auto BlueprintObj = DuplicateObject<UBlueprint>(InBlueprintObj, GetTransientPackage(), *InBlueprintObj->GetName());
+			{
+				auto Reinstancer = FBlueprintCompileReinstancer::Create(BlueprintObj->GeneratedClass);
 
-			IKismetCompilerInterface& Compiler = FModuleManager::LoadModuleChecked<IKismetCompilerInterface>(KISMET_COMPILER_MODULENAME);
+				IKismetCompilerInterface& Compiler = FModuleManager::LoadModuleChecked<IKismetCompilerInterface>(KISMET_COMPILER_MODULENAME);
 
-			TGuardValue<bool> GuardTemplateNameFlag(GCompilingBlueprint, true);
-			FCompilerResultsLog Results;
+				TGuardValue<bool> GuardTemplateNameFlag(GCompilingBlueprint, true);
+				FCompilerResultsLog Results;
 
-			FKismetCompilerOptions CompileOptions;
-			CompileOptions.CompileType = EKismetCompileType::Cpp;
-			CompileOptions.OutCppSourceCode = OutCppSource;
-			CompileOptions.OutHeaderSourceCode = OutHeaderSource;
-			CompileOptions.NewCppClassName = OptionalClassName;
-			Compiler.CompileBlueprint(BlueprintObj, CompileOptions, Results);
+				FKismetCompilerOptions CompileOptions;
+				CompileOptions.CompileType = EKismetCompileType::Cpp;
+				CompileOptions.OutCppSourceCode = OutCppSource;
+				CompileOptions.OutHeaderSourceCode = OutHeaderSource;
+				Compiler.CompileBlueprint(BlueprintObj, CompileOptions, Results);
+			}
+			BlueprintObj->RemoveGeneratedClasses();
+			BlueprintObj->ClearFlags(RF_Standalone);
+			BlueprintObj->MarkPendingKill();
 		}
-		BlueprintObj->RemoveGeneratedClasses();
-		BlueprintObj->ClearFlags(RF_Standalone);
-		BlueprintObj->MarkPendingKill();
+	}
+	else if ((UDEnum || UDStruct) && OutHeaderSource.IsValid())
+	{
+		IKismetCompilerInterface& Compiler = FModuleManager::LoadModuleChecked<IKismetCompilerInterface>(KISMET_COMPILER_MODULENAME);
+		if (UDEnum)
+		{
+			*OutHeaderSource = Compiler.GenerateCppCodeForEnum(UDEnum);
+		}
+		else if (UDStruct)
+		{
+			*OutHeaderSource = Compiler.GenerateCppCodeForStruct(UDStruct);
+		}
+	}
+	else
+	{
+		ensure(false);
 	}
 }
 
diff --git a/Engine/Source/Editor/UnrealEd/Public/Kismet2/KismetEditorUtilities.h b/Engine/Source/Editor/UnrealEd/Public/Kismet2/KismetEditorUtilities.h
--- a/Engine/Source/Editor/UnrealEd/Public/Kismet2/KismetEditorUtilities.h
+++ b/Engine/Source/Editor/UnrealEd/Public/Kismet2/KismetEditorUtilities.h
@@ -105,7 +105,7 @@ public:
 	/** Recompiles the bytecode of a blueprint only.  Should only be run for recompiling dependencies during compile on load */
 	static void RecompileBlueprintBytecode(UBlueprint* BlueprintObj, TArray<UObject*>* ObjLoaded = NULL);
 
-	static void GenerateCppCode(UBlueprint* BlueprintObj, TSharedPtr<FString> OutHeaderSource, TSharedPtr<FString> OutCppSource, const FString& OptionalClassName = FString());
+	static void GenerateCppCode(UObject* Obj, TSharedPtr<FString> OutHeaderSource, TSharedPtr<FString> OutCppSource);
 
 	/** Tries to make sure that a data-only blueprint is conformed to its native parent, in case any native class flags have changed */
 	static void ConformBlueprintFlagsAndComponents(UBlueprint* BlueprintObj);
diff --git a/Engine/Source/Runtime/Engine/Classes/Engine/Blueprint.h b/Engine/Source/Runtime/Engine/Classes/Engine/Blueprint.h
--- a/Engine/Source/Runtime/Engine/Classes/Engine/Blueprint.h
+++ b/Engine/Source/Runtime/Engine/Classes/Engine/Blueprint.h
@@ -82,7 +82,6 @@ public:
 
 	TSharedPtr<FString> OutHeaderSourceCode;
 	TSharedPtr<FString> OutCppSourceCode;
-	FString NewCppClassName;
 
 	bool DoesRequireCppCodeGeneration() const
 	{
