UEBP-40 C++ code generation (WIP):
- Including headers
- Safe context for getters/setters
- improved literals and default values
- basic Interface handling
- better OOB array handling

diff --git a/Engine/Source/Developer/BlueprintCompilerCppBackend/BlueprintCompilerCppBackend.Build.cs b/Engine/Source/Developer/BlueprintCompilerCppBackend/BlueprintCompilerCppBackend.Build.cs
--- a/Engine/Source/Developer/BlueprintCompilerCppBackend/BlueprintCompilerCppBackend.Build.cs
+++ b/Engine/Source/Developer/BlueprintCompilerCppBackend/BlueprintCompilerCppBackend.Build.cs
@@ -26,6 +26,7 @@ namespace UnrealBuildTool.Rules
                     "Engine",
 					"KismetCompiler",
                     "UnrealEd",
+                    "BlueprintGraph",
 				}
 				);
 		}
diff --git a/Engine/Source/Developer/BlueprintCompilerCppBackend/Private/BlueprintCompilerCppBackend.cpp b/Engine/Source/Developer/BlueprintCompilerCppBackend/Private/BlueprintCompilerCppBackend.cpp
--- a/Engine/Source/Developer/BlueprintCompilerCppBackend/Private/BlueprintCompilerCppBackend.cpp
+++ b/Engine/Source/Developer/BlueprintCompilerCppBackend/Private/BlueprintCompilerCppBackend.cpp
@@ -4,152 +4,12 @@
 #include "BlueprintCompilerCppBackend.h"
 #include "BlueprintCompilerCppBackendUtils.h"
 
-FString FBlueprintCompilerCppBackend::TermToText(const FBPTerminal* Term, const UProperty* CoerceProperty)
+FString FBlueprintCompilerCppBackend::TermToText(const FBPTerminal* Term, bool bUseSafeContext)
 {
 	const FString PSC_Self(TEXT("self"));
 	if (Term->bIsLiteral)
 	{
-		//@TODO: Must have a coercion type if it's a literal, because the symbol table isn't plumbed in here and the literals don't carry type information either, yay!
-		ensure(CoerceProperty);
-
-		if (CoerceProperty->IsA(UStrProperty::StaticClass()))
-		{
-			return FString::Printf(TEXT("TEXT(\"%s\")"), *(Term->Name));
-		}
-		else if (CoerceProperty->IsA(UTextProperty::StaticClass()))
-		{
-			return FString::Printf(TEXT("FText::FromString(TEXT(\"%s\"))"), *(Term->Name));
-		}
-		else if (CoerceProperty->IsA(UFloatProperty::StaticClass()))
-		{
-			float Value = FCString::Atof(*(Term->Name));
-			return FString::Printf(TEXT("%f"), Value);
-		}
-		else if (CoerceProperty->IsA(UIntProperty::StaticClass()))
-		{
-			int32 Value = FCString::Atoi(*(Term->Name));
-			return FString::Printf(TEXT("%d"), Value);
-		}
-		else if (auto ByteProperty = Cast<const UByteProperty>(CoerceProperty))
-		{
-			// The PinSubCategoryObject check is to allow enum literals communicate with byte properties as literals
-			auto TypeEnum = ByteProperty->Enum ? ByteProperty->Enum : Cast<UEnum>(Term->Type.PinSubCategoryObject.Get());
-			if (TypeEnum)
-			{
-				return FString::Printf(TEXT("%s::%s"), *TypeEnum->GetName(), *Term->Name);
-			}
-			else
-			{
-				uint8 Value = FCString::Atoi(*(Term->Name));
-				return FString::Printf(TEXT("%u"), Value);
-			}
-		}
-		else if (auto BoolProperty = Cast<const UBoolProperty>(CoerceProperty))
-		{
-			bool bValue = Term->Name.ToBool();
-			return bValue ? TEXT("true") : TEXT("false");
-		}
-		else if (auto NameProperty = Cast<const UNameProperty>(CoerceProperty))
-		{
-			FName LiteralName(*(Term->Name));
-			return FString::Printf(TEXT("FName(TEXT(\"%s\"))"), *(LiteralName.ToString()));
-		}
-		else if (auto StructProperty = Cast<const UStructProperty>(CoerceProperty))
-		{
-			if (StructProperty->Struct == TBaseStructure<FVector>::Get())
-			{
-				FVector Vect = FVector::ZeroVector;
-				FDefaultValueHelper::ParseVector(Term->Name, /*out*/ Vect);
-
-				return FString::Printf(TEXT("FVector(%f,%f,%f)"), Vect.X, Vect.Y, Vect.Z);
-			}
-			else if (StructProperty->Struct == TBaseStructure<FRotator>::Get())
-			{
-				FRotator Rot = FRotator::ZeroRotator;
-				FDefaultValueHelper::ParseRotator(Term->Name, /*out*/ Rot);
-
-				return FString::Printf(TEXT("FRotator(%f,%f,%f)"), Rot.Pitch, Rot.Yaw, Rot.Roll);
-			}
-			else if (StructProperty->Struct == TBaseStructure<FTransform>::Get())
-			{
-				FTransform Trans = FTransform::Identity;
-				Trans.InitFromString(Term->Name);
-
-				const FQuat Rot = Trans.GetRotation();
-				const FVector Translation = Trans.GetTranslation();
-				const FVector Scale = Trans.GetScale3D();
-
-				return FString::Printf(TEXT("FTransform( FQuat(%f,%f,%f,%f), FVector(%f,%f,%f), FVector(%f,%f,%f) )"),
-					Rot.X, Rot.Y, Rot.Z, Rot.W, Translation.X, Translation.Y, Translation.Z, Scale.X, Scale.Y, Scale.Z);
-			}
-			else if (StructProperty->Struct == TBaseStructure<FLinearColor>::Get())
-			{
-				FLinearColor LinearColor;
-				LinearColor.InitFromString(Term->Name);
-				return FString::Printf(TEXT("FLinearColor(%f,%f,%f,%f)"), LinearColor.R, LinearColor.G, LinearColor.B, LinearColor.A);
-			}
-			else
-			{
-				//@todo:  This needs to be more robust, since import text isn't really proper for struct construction.
-				return FString(TEXT("F")) + StructProperty->Struct->GetName() + Term->Name;
-			}
-		}
-		else if (auto ClassProperty = Cast<const UClassProperty>(CoerceProperty))
-		{
-			if (auto FoundClass = Cast<const UClass>(Term->ObjectLiteral))
-			{
-				return FString::Printf(TEXT("%s%s::StaticClass()"), FoundClass->GetPrefixCPP(), *FoundClass->GetName());
-			}
-			return FString(TEXT("NULL"));
-		}
-		else if (CoerceProperty->IsA(UObjectPropertyBase::StaticClass()))
-		{
-			if (Term->Type.PinSubCategory == PSC_Self)
-			{
-				return TEXT("this");
-			}
-			else if (Term->ObjectLiteral)
-			{
-				if (auto AssetClassProp = Cast<UAssetObjectProperty>(CoerceProperty))
-				{
-					return FString::Printf(TEXT("FStringAssetReference(TEXT(\"%s\"))"), *(Term->ObjectLiteral->GetPathName()));
-				}
-				else if (auto LiteralClass = Cast<const UClass>(Term->ObjectLiteral))
-				{
-					return FString::Printf(TEXT("%s%s::StaticClass()"), LiteralClass->GetPrefixCPP(), *LiteralClass->GetName());
-				}
-				else
-				{
-					auto ObjectCoerceProperty = CastChecked<const UObjectPropertyBase>(CoerceProperty);
-					UClass* FoundClass = ObjectCoerceProperty ? ObjectCoerceProperty->PropertyClass : NULL;
-					FString ClassString = FoundClass ? (FString(FoundClass->GetPrefixCPP()) + FoundClass->GetName()) : TEXT("UObject");
-					return FString::Printf(TEXT("FindObject<%s>(ANY_PACKAGE, TEXT(\"%s\"))"), *ClassString, *(Term->ObjectLiteral->GetPathName()));
-				}
-			}
-			else
-			{
-				return FString(TEXT("NULL"));
-			}
-		}
-		else if (CoerceProperty->IsA(UInterfaceProperty::StaticClass()))
-		{
-			if (Term->Type.PinSubCategory == PSC_Self)
-			{
-				return TEXT("this");
-			}
-			else
-			{
-				ensureMsgf(false, TEXT("It is not possible to express this interface property as a literal value!"));
-				return Term->Name;
-			}
-		}
-		else
-			// else if (CoerceProperty->IsA(UMulticastDelegateProperty::StaticClass()))
-			// Cannot assign a literal to a multicast delegate; it should be added instead of assigned
-		{
-			ensureMsgf(false, TEXT("It is not possible to express this type as a literal value!"));
-			return Term->Name;
-		}
+		return FEmitHelper::LiteralTerm(Term->Type, Term->Name, Term->ObjectLiteral);
 	}
 	else if (Term->InlineGeneratedParameter)
 	{
@@ -169,23 +29,31 @@ FString FBlueprintCompilerCppBackend::TermToText(const FBPTerminal* Term, const
 	}
 	else
 	{
-		FString Prefix(TEXT(""));
+		FString ResultPath(TEXT(""));
 		if ((Term->Context != NULL) && (Term->Context->Name != PSC_Self))
 		{
-			Prefix = TermToText(Term->Context);
+			ResultPath = TermToText(Term->Context, false);
 
 			if (Term->Context->IsStructContextType())
 			{
-				Prefix += TEXT(".");
+				ResultPath += TEXT(".");
 			}
 			else
 			{
-				Prefix += TEXT("->");
+				ResultPath += TEXT("->");
 			}
 		}
 
-		Prefix += Term->AssociatedVarProperty ? Term->AssociatedVarProperty->GetNameCPP() : Term->Name;
-		return Prefix;
+		FString Conditions;
+		if (bUseSafeContext)
+		{
+			Conditions = FSafeContextScopedEmmitter::ValidationChain(Term->Context, *this);
+		}
+
+		ResultPath += Term->AssociatedVarProperty ? Term->AssociatedVarProperty->GetNameCPP() : Term->Name;
+		return Conditions.IsEmpty()
+			? ResultPath
+			: FString::Printf(TEXT("((%s) ? (%s) : (%s))"), *Conditions, *ResultPath, *FEmitHelper::DefaultValue(Term->Type));
 	}
 }
 
@@ -239,7 +107,7 @@ void FBlueprintCompilerCppBackend::EmitStructProperties(FStringOutputDevice& Tar
 		}
 		Emit(Header, TEXT(")\n"));
 		Emit(Header, TEXT("\t"));
-		Property->ExportCppDeclaration(Target, EExportedDeclaration::Member, NULL, EPropertyExportCPPFlags::CPPF_CustomTypeName);
+		Property->ExportCppDeclaration(Target, EExportedDeclaration::Member, NULL, EPropertyExportCPPFlags::CPPF_CustomTypeName | EPropertyExportCPPFlags::CPPF_BlueprintCppBackend);
 		Emit(Header, TEXT(";\n"));
 	}
 }
@@ -289,15 +157,36 @@ void FBlueprintCompilerCppBackend::GenerateCodeFromClass(UClass* SourceClass, TI
 		}
 	}
 
-	Emit(Header, TEXT("UCLASS(Blueprintable, BlueprintType)\n"));
+	// Class declaration
+	const bool bIsInterface = SourceClass->IsChildOf<UInterface>();
+	if (bIsInterface)
+	{
+		Emit(Header, *FString::Printf(TEXT("UINTERFACE(Blueprintable)\nclass U%s : public UInterface\n{\n\tGENERATED_BODY()\n};\n"), *CleanCppClassName));
+		Emit(Header, *FString::Printf(TEXT("\nclass I%s"), *CleanCppClassName));
+	}
+	else
+	{
+		Emit(Header, TEXT("UCLASS("));
+		if (!SourceClass->IsChildOf<UBlueprintFunctionLibrary>())
+		{
+			Emit(Header, TEXT("Blueprintable, BlueprintType"));
+		}
+		Emit(Header, TEXT(")\n"));
 
-	UClass* SuperClass = SourceClass->GetSuperClass();
-	Emit(Header,
-		*FString::Printf(TEXT("class %s : public %s%s\n"), *CppClassName, SuperClass->GetPrefixCPP(), *SuperClass->GetName()));
-	Emit(Header,
-		TEXT("{\n")
-		TEXT("public:\n"));
-	Emit(Header, *FString::Printf(TEXT("\tGENERATED_BODY()\n")));
+		UClass* SuperClass = SourceClass->GetSuperClass();
+		Emit(Header, *FString::Printf(TEXT("class %s : public %s%s"), *CppClassName, SuperClass->GetPrefixCPP(), *SuperClass->GetName()));
+
+		for (auto& ImplementedInterface : SourceClass->Interfaces)
+		{
+			if (ImplementedInterface.Class)
+			{
+				Emit(Header, *FString::Printf(TEXT(", public I%s"), *ImplementedInterface.Class->GetName()));
+			}
+		}
+	}
+
+	// Begin scope
+	Emit(Header,TEXT("\n{\npublic:\n\tGENERATED_BODY()\n"));
 
 	EmitStructProperties(Header, SourceClass);
 
@@ -314,13 +203,7 @@ void FBlueprintCompilerCppBackend::GenerateCodeFromClass(UClass* SourceClass, TI
 		Emit(Header, TEXT("\n"));
 	}
 
-	TArray<FString> PersistentHeaders;
-	PersistentHeaders.Add(FString(FApp::GetGameName()) + TEXT(".h"));
-	PersistentHeaders.Add(CleanCppClassName + TEXT(".h"));
-	PersistentHeaders.Add(TEXT("GeneratedCodeHelpers.h"));
-	Emit(Body, *FEmitHelper::GatherHeadersToInclude(SourceClass, PersistentHeaders));
-
-	//constructor is generated by GENERATED_BODY macto
+	//constructor is generated by GENERATED_BODY macro
 	// Emit(Body, *FString::Printf(TEXT("%s::%s(const FObjectInitializer& ObjectInitializer) : Super(ObjectInitializer) {}\n\n"), *CppClassName, *CppClassName));
 
 	for (int32 i = 0; i < Functions.Num(); ++i)
@@ -343,7 +226,7 @@ void FBlueprintCompilerCppBackend::EmitCallDelegateStatment(FKismetFunctionConte
 
 	Emit(Body, TEXT("\t\t\t"));
 	Emit(Body, *SafeContextScope.GetAdditionalIndent());
-	Emit(Body, *FString::Printf(TEXT("%s.Broadcast(%s)"), *TermToText(Statement.FunctionContext), *EmitMethodInputParameterList(Statement)));
+	Emit(Body, *FString::Printf(TEXT("%s.Broadcast(%s)"), *TermToText(Statement.FunctionContext, false), *EmitMethodInputParameterList(Statement)));
 }
 
 FString FBlueprintCompilerCppBackend::EmitMethodInputParameterList(FBlueprintCompiledStatement& Statement)
@@ -385,7 +268,7 @@ FString FBlueprintCompilerCppBackend::EmitMethodInputParameterList(FBlueprintCom
 			else
 			{
 				// Emit a normal parameter term
-				VarName = TermToText(Term, FuncParamProperty);
+				VarName = TermToText(Term);
 			}
 
 			if (FuncParamProperty->HasAnyPropertyFlags(CPF_OutParm))
@@ -408,47 +291,48 @@ FString FBlueprintCompilerCppBackend::EmitCallStatmentInner(FBlueprintCompiledSt
 	const bool bCallOnDifferentObject = Statement.FunctionContext && (Statement.FunctionContext->Name != TEXT("self"));
 	const bool bStaticCall = Statement.FunctionToCall->HasAnyFunctionFlags(FUNC_Static);
 	const bool bUseSafeContext = bCallOnDifferentObject && !bStaticCall;
-	FSafeContextScopedEmmitter SafeContextScope(Result, bUseSafeContext ? Statement.FunctionContext : nullptr, *this, TEXT("\t\t\t"));
-	ensure(!bInline || !SafeContextScope.IsSafeContextUsed());
-
-	if (!bInline)
 	{
-		Result += TEXT("\t\t\t");
-		Result += SafeContextScope.GetAdditionalIndent();
+		FSafeContextScopedEmmitter SafeContextScope(Result, bUseSafeContext ? Statement.FunctionContext : nullptr, *this, TEXT("\t\t\t"));
+		ensure(!bInline || !SafeContextScope.IsSafeContextUsed());
 
-		// Handle the return value of the function being called
-		UProperty* FuncToCallReturnProperty = Statement.FunctionToCall->GetReturnProperty();
-		if (FuncToCallReturnProperty && ensure(Statement.LHS))
+		if (!bInline)
 		{
-			Result += FString::Printf(TEXT("%s = "), *TermToText(Statement.LHS));
+			Result += TEXT("\t\t\t");
+			Result += SafeContextScope.GetAdditionalIndent();
+
+			// Handle the return value of the function being called
+			UProperty* FuncToCallReturnProperty = Statement.FunctionToCall->GetReturnProperty();
+			if (FuncToCallReturnProperty && ensure(Statement.LHS))
+			{
+				Result += FString::Printf(TEXT("%s = "), *TermToText(Statement.LHS));
+			}
 		}
-	}
 
-	// Emit object to call the method on
-	if (bStaticCall)
-	{
-		const bool bIsCustomThunk = Statement.FunctionToCall->HasMetaData(TEXT("CustomStructureParam")) || Statement.FunctionToCall->HasMetaData(TEXT("ArrayParm"));
-		auto OwnerClass = Statement.FunctionToCall->GetOuterUClass();
-		Result += bIsCustomThunk ? TEXT("FCustomThunkTemplates::") : FString::Printf(TEXT("%s%s::"), OwnerClass->GetPrefixCPP(), *OwnerClass->GetName());
-	}
-	else if(bCallOnDifferentObject) //@TODO: Badness, could be a self reference wired to another instance!
-	{
-		Result += FString::Printf(TEXT("%s->"), *TermToText(Statement.FunctionContext, (UProperty*)(GetDefault<UObjectProperty>())));
-	}
+		// Emit object to call the method on
+		if (bStaticCall)
+		{
+			const bool bIsCustomThunk = Statement.FunctionToCall->HasMetaData(TEXT("CustomStructureParam")) || Statement.FunctionToCall->HasMetaData(TEXT("ArrayParm"));
+			auto OwnerClass = Statement.FunctionToCall->GetOuterUClass();
+			Result += bIsCustomThunk ? TEXT("FCustomThunkTemplates::") : FString::Printf(TEXT("%s%s::"), OwnerClass->GetPrefixCPP(), *OwnerClass->GetName());
+		}
+		else if (bCallOnDifferentObject) //@TODO: Badness, could be a self reference wired to another instance!
+		{
+			Result += FString::Printf(TEXT("%s->"), *TermToText(Statement.FunctionContext, false));
+		}
 
-	// Emit method name
-	Result += FString::Printf(TEXT("%s%s"), Statement.bIsParentContext ? TEXT("Super::") : TEXT(""), *Statement.FunctionToCall->GetName());
+		// Emit method name
+		Result += FString::Printf(TEXT("%s%s"), Statement.bIsParentContext ? TEXT("Super::") : TEXT(""), *Statement.FunctionToCall->GetName());
 
-	// Emit method parameter list
-	Result += TEXT("(");
-	Result += EmitMethodInputParameterList(Statement);
-	Result += TEXT(")");
+		// Emit method parameter list
+		Result += TEXT("(");
+		Result += EmitMethodInputParameterList(Statement);
+		Result += TEXT(")");
 
-	if (!bInline)
-	{
-		Result += TEXT(";\n");
+		if (!bInline)
+		{
+			Result += TEXT(";\n");
+		}
 	}
-
 	return Result;
 }
 
@@ -467,29 +351,29 @@ FString FBlueprintCompilerCppBackend::EmitSwitchValueStatmentInner(FBlueprintCom
 	auto DefaultValueTerm = Statement.RHS.Last();
 
 	FString Result = FString::Printf(TEXT("TSwitchValue(%s, %s, %d")
-		, *TermToText(IndexTerm, nullptr) //index
-		, *TermToText(DefaultValueTerm, nullptr) // default
+		, *TermToText(IndexTerm) //index
+		, *TermToText(DefaultValueTerm) // default
 		, NumCases);
 	
-	UProperty* VirtualIndexProperty = IndexTerm ? IndexTerm->AssociatedVarProperty : nullptr;
-	check(VirtualIndexProperty);
+	const uint32 CppTemplateTypeFlags = EPropertyExportCPPFlags::CPPF_CustomTypeName 
+		| EPropertyExportCPPFlags::CPPF_NoConst | EPropertyExportCPPFlags::CPPF_NoRef 
+		| EPropertyExportCPPFlags::CPPF_BlueprintCppBackend;
+
 	FStringOutputDevice IndexDeclaration;
-	VirtualIndexProperty->ExportCppDeclaration(IndexDeclaration, EExportedDeclaration::Parameter, NULL
-		, EPropertyExportCPPFlags::CPPF_CustomTypeName | EPropertyExportCPPFlags::CPPF_NoConst | EPropertyExportCPPFlags::CPPF_NoRef, true);
+	check(IndexTerm && IndexTerm->AssociatedVarProperty);
+	IndexTerm->AssociatedVarProperty->ExportCppDeclaration(IndexDeclaration, EExportedDeclaration::Parameter, NULL, CppTemplateTypeFlags, true);
 
-	UProperty* VirtualValueProperty = DefaultValueTerm ? DefaultValueTerm->AssociatedVarProperty : nullptr;
-	check(VirtualValueProperty);
+	check(DefaultValueTerm && DefaultValueTerm->AssociatedVarProperty);
 	FStringOutputDevice ValueDeclaration;
-	VirtualValueProperty->ExportCppDeclaration(ValueDeclaration, EExportedDeclaration::Parameter, NULL
-		, EPropertyExportCPPFlags::CPPF_CustomTypeName | EPropertyExportCPPFlags::CPPF_NoConst | EPropertyExportCPPFlags::CPPF_NoRef, true);
+	DefaultValueTerm->AssociatedVarProperty->ExportCppDeclaration(ValueDeclaration, EExportedDeclaration::Parameter, NULL, CppTemplateTypeFlags, true);
 
 	for (int32 TermIndex = TermsBeforeCases; TermIndex < (NumCases * TermsPerCase); TermIndex += TermsPerCase)
 	{
 		Result += FString::Printf(TEXT(", TSwitchPair<%s, %s>(%s, %s)")
 			, *IndexDeclaration
 			, *ValueDeclaration
-			, *TermToText(Statement.RHS[TermIndex], VirtualIndexProperty)
-			, *TermToText(Statement.RHS[TermIndex + 1], VirtualValueProperty));
+			, *TermToText(Statement.RHS[TermIndex])
+			, *TermToText(Statement.RHS[TermIndex + 1]));
 	}
 
 	Result += TEXT(")");
@@ -499,8 +383,8 @@ FString FBlueprintCompilerCppBackend::EmitSwitchValueStatmentInner(FBlueprintCom
 
 void FBlueprintCompilerCppBackend::EmitAssignmentStatment(FKismetFunctionContext& FunctionContext, FBlueprintCompiledStatement& Statement)
 {
-	FString DestinationExpression = TermToText(Statement.LHS);
-	FString SourceExpression = TermToText(Statement.RHS[0], Statement.LHS->AssociatedVarProperty);
+	FString DestinationExpression = TermToText(Statement.LHS, false);
+	FString SourceExpression = TermToText(Statement.RHS[0]);
 
 	FSafeContextScopedEmmitter SafeContextScope(Body, Statement.LHS->Context, *this, TEXT("\t\t\t"));
 	// Emit the assignment statement
@@ -511,9 +395,9 @@ void FBlueprintCompilerCppBackend::EmitAssignmentStatment(FKismetFunctionContext
 
 void FBlueprintCompilerCppBackend::EmitCastObjToInterfaceStatement(FKismetFunctionContext& FunctionContext, FBlueprintCompiledStatement& Statement)
 {
-	FString InterfaceClass = TermToText(Statement.RHS[0], (UProperty*)(GetDefault<UClassProperty>()));
-	FString ObjectValue = TermToText(Statement.RHS[1], (UProperty*)(GetDefault<UObjectProperty>()));
-	FString InterfaceValue = TermToText(Statement.LHS, (UProperty*)(GetDefault<UObjectProperty>()));
+	FString InterfaceClass = TermToText(Statement.RHS[0]);
+	FString ObjectValue = TermToText(Statement.RHS[1]);
+	FString InterfaceValue = TermToText(Statement.LHS);
 
 	Emit(Body, *FString::Printf(TEXT("\t\t\tif ( IsValid(%s) && %s->GetClass()->ImplementsInterface(%s) )\n"), *ObjectValue, *ObjectValue, *InterfaceClass));
 	Emit(Body, *FString::Printf(TEXT("\t\t\t{\n")));
@@ -529,9 +413,9 @@ void FBlueprintCompilerCppBackend::EmitCastObjToInterfaceStatement(FKismetFuncti
 
 void FBlueprintCompilerCppBackend::EmitCastBetweenInterfacesStatement(FKismetFunctionContext& FunctionContext, FBlueprintCompiledStatement& Statement)
 {
-	FString ClassToCastTo = TermToText(Statement.RHS[0], (UProperty*)(GetDefault<UClassProperty>()));
-	FString InputInterface = TermToText(Statement.RHS[1], (UProperty*)(GetDefault<UInterfaceProperty>()));
-	FString ResultInterface = TermToText(Statement.LHS, (UProperty*)(GetDefault<UInterfaceProperty>()));
+	FString ClassToCastTo = TermToText(Statement.RHS[0]);
+	FString InputInterface = TermToText(Statement.RHS[1]);
+	FString ResultInterface = TermToText(Statement.LHS);
 
 	FString InputObject = FString::Printf(TEXT("%s.GetObjectRef()"), *InputInterface);
 
@@ -549,9 +433,9 @@ void FBlueprintCompilerCppBackend::EmitCastBetweenInterfacesStatement(FKismetFun
 
 void FBlueprintCompilerCppBackend::EmitCastInterfaceToObjStatement(FKismetFunctionContext& FunctionContext, FBlueprintCompiledStatement& Statement)
 {
-	FString ClassToCastTo = TermToText(Statement.RHS[0], (UProperty*)(GetDefault<UClassProperty>()));
-	FString InputInterface = TermToText(Statement.RHS[1], (UProperty*)(GetDefault<UInterfaceProperty>()));
-	FString ResultObject = TermToText(Statement.LHS, (UProperty*)(GetDefault<UInterfaceProperty>()));
+	FString ClassToCastTo = TermToText(Statement.RHS[0]);
+	FString InputInterface = TermToText(Statement.RHS[1]);
+	FString ResultObject = TermToText(Statement.LHS);
 
 	FString InputObject = FString::Printf(TEXT("%s.GetObjectRef()"), *InputInterface);
 
@@ -561,9 +445,9 @@ void FBlueprintCompilerCppBackend::EmitCastInterfaceToObjStatement(FKismetFuncti
 
 void FBlueprintCompilerCppBackend::EmitDynamicCastStatement(FKismetFunctionContext& FunctionContext, FBlueprintCompiledStatement& Statement)
 {
-	FString TargetClass = TermToText(Statement.RHS[0], (UProperty*)(GetDefault<UClassProperty>()));
-	FString ObjectValue = TermToText(Statement.RHS[1], (UProperty*)(GetDefault<UObjectProperty>()));
-	FString CastedValue = TermToText(Statement.LHS, (UProperty*)(GetDefault<UObjectProperty>()));
+	FString TargetClass = TermToText(Statement.RHS[0]);
+	FString ObjectValue = TermToText(Statement.RHS[1]);
+	FString CastedValue = TermToText(Statement.LHS);
 
 	Emit(Body, *FString::Printf(TEXT("\t\t\t%s = Cast<%s>(%s);\n"),
 		*CastedValue, *TargetClass, *ObjectValue));
@@ -571,9 +455,9 @@ void FBlueprintCompilerCppBackend::EmitDynamicCastStatement(FKismetFunctionConte
 
 void FBlueprintCompilerCppBackend::EmitMetaCastStatement(FKismetFunctionContext& FunctionContext, FBlueprintCompiledStatement& Statement)
 {
-	FString DesiredClass = TermToText(Statement.RHS[0], (UProperty*)(GetDefault<UClassProperty>()));
-	FString SourceClass = TermToText(Statement.RHS[1], (UProperty*)(GetDefault<UClassProperty>()));
-	FString Destination = TermToText(Statement.LHS, (UProperty*)(GetDefault<UClassProperty>()));
+	FString DesiredClass = TermToText(Statement.RHS[0]);
+	FString SourceClass = TermToText(Statement.RHS[1]);
+	FString Destination = TermToText(Statement.LHS);
 
 	Emit(Body, *FString::Printf(TEXT("\t\t\t%s = DynamicMetaCast(%s, %s);\n"),
 		*Destination, *DesiredClass, *SourceClass));
@@ -595,7 +479,7 @@ void FBlueprintCompilerCppBackend::EmitAddMulticastDelegateStatement(FKismetFunc
 	Emit(Body, TEXT("\t\t\t"));
 	Emit(Body, *SafeContextScope.GetAdditionalIndent());
 
-	const FString Delegate = TermToText(Statement.LHS);
+	const FString Delegate = TermToText(Statement.LHS, false);
 	const FString DelegateToAdd = TermToText(Statement.RHS[0]);
 
 	Emit(Body, *FString::Printf(TEXT("%s.Add(%s);\n"), *Delegate, *DelegateToAdd));
@@ -609,7 +493,7 @@ void FBlueprintCompilerCppBackend::EmitRemoveMulticastDelegateStatement(FKismetF
 	Emit(Body, TEXT("\t\t\t"));
 	Emit(Body, *SafeContextScope.GetAdditionalIndent());
 
-	const FString Delegate = TermToText(Statement.LHS);
+	const FString Delegate = TermToText(Statement.LHS, false);
 	const FString DelegateToAdd = TermToText(Statement.RHS[0]);
 
 	Emit(Body, *FString::Printf(TEXT("%s.Remove(%s);\n"), *Delegate, *DelegateToAdd));
@@ -625,9 +509,9 @@ void FBlueprintCompilerCppBackend::EmitBindDelegateStatement(FKismetFunctionCont
 	Emit(Body, *SafeContextScope.GetAdditionalIndent());
 
 
-	const FString Delegate = TermToText(Statement.LHS);
-	const FString NameTerm = TermToText(Statement.RHS[0], GetDefault<UNameProperty>());
-	const FString ObjectTerm = TermToText(Statement.RHS[1], GetDefault<UObjectProperty>());
+	const FString Delegate = TermToText(Statement.LHS, false);
+	const FString NameTerm = TermToText(Statement.RHS[0]);
+	const FString ObjectTerm = TermToText(Statement.RHS[1]);
 
 	Emit(Body, *FString::Printf(TEXT("%s.BindUFunction(%s,%s);\n"), *Delegate, *ObjectTerm, *NameTerm));
 }
@@ -640,7 +524,7 @@ void FBlueprintCompilerCppBackend::EmitClearMulticastDelegateStatement(FKismetFu
 	Emit(Body, TEXT("\t\t\t"));
 	Emit(Body, *SafeContextScope.GetAdditionalIndent());
 
-	const FString Delegate = TermToText(Statement.LHS);
+	const FString Delegate = TermToText(Statement.LHS, false);
 
 	Emit(Body, *FString::Printf(TEXT("%s.Clear();\n"), *Delegate));
 }
@@ -656,12 +540,7 @@ void FBlueprintCompilerCppBackend::EmitCreateArrayStatement(FKismetFunctionConte
 	for (int32 i = 0; i < Statement.RHS.Num(); ++i)
 	{
 		FBPTerminal* CurrentTerminal = Statement.RHS[i];
-		Emit(Body,
-			*FString::Printf(
-			TEXT("\t\t\t%s[%d] = %s;"),
-			*Array,
-			i,
-			*TermToText(CurrentTerminal, (CurrentTerminal->bIsLiteral ? InnerProperty : NULL))));
+		Emit(Body,*FString::Printf(TEXT("\t\t\t%s[%d] = %s;"), *Array, i, *TermToText(CurrentTerminal)));
 	}
 }
 
@@ -670,7 +549,7 @@ void FBlueprintCompilerCppBackend::EmitGotoStatement(FKismetFunctionContext& Fun
 	if (Statement.Type == KCST_ComputedGoto)
 	{
 		FString NextStateExpression;
-		NextStateExpression = TermToText(Statement.LHS, (UProperty*)(GetDefault<UIntProperty>()));
+		NextStateExpression = TermToText(Statement.LHS);
 
 		Emit(Body, *FString::Printf(TEXT("\t\t\tCurrentState = %s;\n"), *NextStateExpression));
 		Emit(Body, *FString::Printf(TEXT("\t\t\tbreak;\n")));
@@ -678,7 +557,7 @@ void FBlueprintCompilerCppBackend::EmitGotoStatement(FKismetFunctionContext& Fun
 	else if ((Statement.Type == KCST_GotoIfNot) || (Statement.Type == KCST_EndOfThreadIfNot) || (Statement.Type == KCST_GotoReturnIfNot))
 	{
 		FString ConditionExpression;
-		ConditionExpression = TermToText(Statement.LHS, (UProperty*)(GetDefault<UBoolProperty>()));
+		ConditionExpression = TermToText(Statement.LHS);
 
 		Emit(Body, *FString::Printf(TEXT("\t\t\tif (!%s)\n"), *ConditionExpression));
 		Emit(Body, *FString::Printf(TEXT("\t\t\t{\n")));
@@ -736,7 +615,7 @@ void FBlueprintCompilerCppBackend::DeclareLocalVariables(FKismetFunctionContext&
 		UProperty* LocalVariable = LocalVariables[i];
 
 		Emit(Body, TEXT("\t"));
-		LocalVariable->ExportCppDeclaration(Body, EExportedDeclaration::Local, NULL, EPropertyExportCPPFlags::CPPF_CustomTypeName);
+		LocalVariable->ExportCppDeclaration(Body, EExportedDeclaration::Local, NULL, EPropertyExportCPPFlags::CPPF_CustomTypeName | EPropertyExportCPPFlags::CPPF_BlueprintCppBackend);
 		Emit(Body, TEXT("{};\n"));
 	}
 
@@ -855,8 +734,8 @@ void FBlueprintCompilerCppBackend::ConstructFunction(FKismetFunctionContext& Fun
 					Emit(Body, TEXT("/*out*/ "));
 				}
 
-				ArgProperty->ExportCppDeclaration(Header, EExportedDeclaration::Parameter, NULL, EPropertyExportCPPFlags::CPPF_CustomTypeName);
-				ArgProperty->ExportCppDeclaration(Body, EExportedDeclaration::Parameter, NULL, EPropertyExportCPPFlags::CPPF_CustomTypeName);
+				ArgProperty->ExportCppDeclaration(Header, EExportedDeclaration::Parameter, NULL, EPropertyExportCPPFlags::CPPF_CustomTypeName | EPropertyExportCPPFlags::CPPF_BlueprintCppBackend);
+				ArgProperty->ExportCppDeclaration(Body, EExportedDeclaration::Parameter, NULL, EPropertyExportCPPFlags::CPPF_CustomTypeName | EPropertyExportCPPFlags::CPPF_BlueprintCppBackend);
 			}
 
 			Emit(Header, TEXT(")"));
@@ -1090,13 +969,21 @@ void FBlueprintCompilerCppBackend::EmitFileBeginning(const FString& CleanName, U
 {
 	Emit(Header, TEXT("#pragma once\n\n"));
 	if (SourceStruct)
-	{
+	{ 
+		{
+			TArray<FString> PersistentHeaders;
+			PersistentHeaders.Add(FString(FApp::GetGameName()) + TEXT(".h"));
+			PersistentHeaders.Add(CleanName + TEXT(".h"));
+			PersistentHeaders.Add(TEXT("GeneratedCodeHelpers.h"));
+			Emit(Body, *FEmitHelper::GatherNativeHeadersToInclude(SourceStruct, PersistentHeaders));
+		}
+
 		// find objects referenced by functions/script
 		TArray<UObject*> IncludeInHeader;
 		TArray<UObject*> IncludeInBody;
 		{
-			FReferenceFinder HeaderReferenceFinder(IncludeInHeader, NULL, false, false, false, false);
-			FReferenceFinder BodyReferenceFinder(IncludeInBody, NULL, false, false, false, false);
+			FReferenceFinder HeaderReferenceFinder(IncludeInHeader, NULL, false, false, true, false);
+			FReferenceFinder BodyReferenceFinder(IncludeInBody, NULL, false, false, true, false);
 			{
 				TArray<UObject*> ObjectsToCheck;
 				GetObjectsWithOuter(SourceStruct, ObjectsToCheck, true);
@@ -1121,17 +1008,44 @@ void FBlueprintCompilerCppBackend::EmitFileBeginning(const FString& CleanName, U
 			{
 				IncludeInHeader.AddUnique(SuperStruct);
 			}
+
+			if (auto SourceClass = Cast<UClass>(SourceStruct))
+			{
+				for (auto& ImplementedInterface : SourceClass->Interfaces)
+				{
+					IncludeInHeader.AddUnique(ImplementedInterface.Class);
+				}
+			}
 		}
 
 		TSet<FString> AlreadyIncluded;
 		AlreadyIncluded.Add(SourceStruct->GetName());
 
-		auto EmitInner = [&](FStringOutputDevice& Dst, TArray<UObject*>& Src)
+		auto EmitInner = [&](FStringOutputDevice& Dst, const TArray<UObject*>& Src)
 		{
 			for (UObject* Obj : Src)
 			{
-				const bool bWantedType = Obj && (Obj->IsA<UBlueprintGeneratedClass>() || Obj->IsA<UUserDefinedEnum>() || Obj->IsA<UUserDefinedStruct>());
-				if (bWantedType)
+				bool bWantedType = Obj && (Obj->IsA<UBlueprintGeneratedClass>() || Obj->IsA<UUserDefinedEnum>() || Obj->IsA<UUserDefinedStruct>());
+				if (!bWantedType)
+				{
+					for (UObject* OuterObj = (Obj ? Obj->GetOuter() : nullptr); nullptr != OuterObj; OuterObj = OuterObj->GetOuter())
+					{
+						if (OuterObj->IsA<UBlueprintGeneratedClass>() || OuterObj->IsA<UUserDefinedStruct>())
+						{
+							Obj = OuterObj;
+							bWantedType = true;
+							break;
+						}
+					}
+				}
+
+				if (!bWantedType)
+				{
+					Obj = Obj ? Cast<UBlueprintGeneratedClass>(Obj->GetClass()) : nullptr;
+					bWantedType = (nullptr != Obj);
+				}
+
+				if (bWantedType && Obj)
 				{
 					const FString Name = Obj->GetName();
 					if (!AlreadyIncluded.Contains(Name))
diff --git a/Engine/Source/Developer/BlueprintCompilerCppBackend/Private/BlueprintCompilerCppBackend.h b/Engine/Source/Developer/BlueprintCompilerCppBackend/Private/BlueprintCompilerCppBackend.h
--- a/Engine/Source/Developer/BlueprintCompilerCppBackend/Private/BlueprintCompilerCppBackend.h
+++ b/Engine/Source/Developer/BlueprintCompilerCppBackend/Private/BlueprintCompilerCppBackend.h
@@ -84,7 +84,7 @@ public:
 	/** Builds both the header declaration and body implementation of a function */
 	void ConstructFunction(FKismetFunctionContext& FunctionContext, bool bGenerateStubOnly);
 
-	FString TermToText(const FBPTerminal* Term, const UProperty* SourceProperty = nullptr);
+	FString TermToText(const FBPTerminal* Term, bool bUseSafeContext = true);
 
 protected:
 	void EmitFileBeginning(const FString& CleanName, UStruct* SourceStruct);
diff --git a/Engine/Source/Developer/BlueprintCompilerCppBackend/Private/BlueprintCompilerCppBackendUtils.h b/Engine/Source/Developer/BlueprintCompilerCppBackend/Private/BlueprintCompilerCppBackendUtils.h
--- a/Engine/Source/Developer/BlueprintCompilerCppBackend/Private/BlueprintCompilerCppBackendUtils.h
+++ b/Engine/Source/Developer/BlueprintCompilerCppBackend/Private/BlueprintCompilerCppBackendUtils.h
@@ -1,5 +1,7 @@
 // Copyright 1998-2015 Epic Games, Inc. All Rights Reserved.
 
+#include "EdGraphSchema_K2.h"
+
 // Generates single if scope. It's condition checks context of given term.
 struct FSafeContextScopedEmmitter
 {
@@ -18,8 +20,7 @@ public:
 		return bSafeContextUsed;
 	}
 
-	FSafeContextScopedEmmitter(FString& InBody, const FBPTerminal* Term, FBlueprintCompilerCppBackend& CppBackend, const TCHAR* InCurrentIndent)
-		: Body(InBody), bSafeContextUsed(false), CurrentIndent(InCurrentIndent)
+	static FString ValidationChain(const FBPTerminal* Term, FBlueprintCompilerCppBackend& CppBackend)
 	{
 		TArray<FString> SafetyConditions;
 		for (; Term; Term = Term->Context)
@@ -27,25 +28,36 @@ public:
 			if (!Term->IsStructContextType() && (Term->Type.PinSubCategory != TEXT("self")))
 			{
 				ensure(!Term->bIsLiteral);
-				SafetyConditions.Add(CppBackend.TermToText(Term));
+				SafetyConditions.Add(CppBackend.TermToText(Term, false));
 			}
 		}
 
-		if (SafetyConditions.Num())
+		FString Result;
+		for (int32 Iter = SafetyConditions.Num() - 1; Iter >= 0; --Iter)
+		{
+			Result += FString(TEXT("IsValid("));
+			Result += SafetyConditions[Iter];
+			Result += FString(TEXT(")"));
+			if (Iter)
+			{
+				Result += FString(TEXT(" && "));
+			}
+		}
+
+		return Result;
+	}
+
+	FSafeContextScopedEmmitter(FString& InBody, const FBPTerminal* Term, FBlueprintCompilerCppBackend& CppBackend, const TCHAR* InCurrentIndent)
+		: Body(InBody), bSafeContextUsed(false), CurrentIndent(InCurrentIndent)
+	{
+		const FString Conditions = ValidationChain(Term, CppBackend);
+
+		if (!Conditions.IsEmpty())
 		{
 			bSafeContextUsed = true;
 			Body += CurrentIndent;
 			Body += TEXT("if (");
-			for (int32 Iter = SafetyConditions.Num() - 1; Iter >= 0; --Iter)
-			{
-				Body += FString(TEXT("IsValid("));
-				Body += SafetyConditions[Iter];
-				Body += FString(TEXT(")"));
-				if (Iter)
-				{
-					Body += FString(TEXT(" && "));
-				}
-			}
+			Body += Conditions;
 			Body += TEXT(")\n");
 			Body += CurrentIndent;
 			Body += TEXT("{\n");
@@ -371,7 +383,7 @@ struct FEmitHelper
 		return Result;
 	}
 
-	static FString GatherHeadersToInclude(UClass* SourceClass, const TArray<FString>& PersistentHeaders)
+	static FString GatherNativeHeadersToInclude(UField* SourceItem, const TArray<FString>& PersistentHeaders)
 	{
 		TSet<FString> HeaderFiles;
 		HeaderFiles.Append(PersistentHeaders);
@@ -380,7 +392,7 @@ struct FEmitHelper
 			{
 				FReferenceFinder ReferenceFinder(ReferencedObjects, NULL, false, false, false, true);
 				TArray<UObject*> ObjectsToCheck;
-				GetObjectsWithOuter(SourceClass, ObjectsToCheck, true);
+				GetObjectsWithOuter(SourceItem, ObjectsToCheck, true);
 				//CDO?
 				for (auto Obj : ObjectsToCheck)
 				{
@@ -390,7 +402,7 @@ struct FEmitHelper
 
 			auto EngineSourceDir = FPaths::EngineSourceDir();
 			auto GameSourceDir = FPaths::GameSourceDir();
-			auto CurrentPackage = SourceClass->GetTypedOuter<UPackage>();
+			auto CurrentPackage = SourceItem->GetTypedOuter<UPackage>();
 			for (auto Obj : ReferencedObjects)
 			{
 				auto Field = Cast<UField>(Obj);
@@ -419,5 +431,157 @@ struct FEmitHelper
 
 		return Result;
 	}
+
+	static FString LiteralTerm(const FEdGraphPinType& Type, const FString& CustomValue, UObject* LiteralObject)
+	{
+		auto Schema = GetDefault<UEdGraphSchema_K2>();
+
+		if (UEdGraphSchema_K2::PC_String == Type.PinCategory)
+		{
+			return FString::Printf(TEXT("TEXT(\"%s\")"), *CustomValue);
+		}
+		else if (UEdGraphSchema_K2::PC_Text == Type.PinCategory)
+		{
+			return FString::Printf(TEXT("FText::FromString(TEXT(\"%s\"))"), *CustomValue);
+		}
+		else if (UEdGraphSchema_K2::PC_Float == Type.PinCategory)
+		{
+			float Value = CustomValue.IsEmpty() ? 0.0f : FCString::Atof(*CustomValue);
+			return FString::Printf(TEXT("%f"), Value);
+		}
+		else if (UEdGraphSchema_K2::PC_Int == Type.PinCategory)
+		{
+			int32 Value = CustomValue.IsEmpty() ? 0 : FCString::Atoi(*CustomValue);
+			return FString::Printf(TEXT("%d"), Value);
+		}
+		else if ((UEdGraphSchema_K2::PC_Byte == Type.PinCategory) || (UEdGraphSchema_K2::PC_Enum == Type.PinCategory))
+		{
+			auto TypeEnum = Cast<UEnum>(Type.PinSubCategoryObject.Get());
+			if (TypeEnum)
+			{
+				return FString::Printf(TEXT("%s::%s"), *TypeEnum->GetName(), CustomValue.IsEmpty() 
+					? *TypeEnum->GetEnumName(TypeEnum->NumEnums()-1)
+					: *CustomValue);
+			}
+			else
+			{
+				uint8 Value = CustomValue.IsEmpty() ? 0 : FCString::Atoi(*CustomValue);
+				return FString::Printf(TEXT("%u"), Value);
+			}
+		}
+		else if (UEdGraphSchema_K2::PC_Boolean == Type.PinCategory)
+		{
+			const bool bValue = CustomValue.ToBool();
+			return bValue ? TEXT("true") : TEXT("false");
+		}
+		else if (UEdGraphSchema_K2::PC_Name == Type.PinCategory)
+		{
+			return CustomValue.IsEmpty() 
+				? TEXT("FName()") 
+				: FString::Printf(TEXT("FName(TEXT(\"%s\"))"), *(FName(*CustomValue).ToString()));
+		}
+		else if (UEdGraphSchema_K2::PC_Struct == Type.PinCategory)
+		{
+			auto StructType = Cast<UScriptStruct>(Type.PinSubCategoryObject.Get());
+			ensure(StructType);
+			if (StructType == TBaseStructure<FVector>::Get())
+			{
+				FVector Vect = FVector::ZeroVector;
+				FDefaultValueHelper::ParseVector(CustomValue, /*out*/ Vect);
+				return FString::Printf(TEXT("FVector(%f,%f,%f)"), Vect.X, Vect.Y, Vect.Z);
+			}
+			else if (StructType == TBaseStructure<FRotator>::Get())
+			{
+				FRotator Rot = FRotator::ZeroRotator;
+				FDefaultValueHelper::ParseRotator(CustomValue, /*out*/ Rot);
+				return FString::Printf(TEXT("FRotator(%f,%f,%f)"), Rot.Pitch, Rot.Yaw, Rot.Roll);
+			}
+			else if (StructType == TBaseStructure<FTransform>::Get())
+			{
+				FTransform Trans = FTransform::Identity;
+				Trans.InitFromString(CustomValue);
+				const FQuat Rot = Trans.GetRotation();
+				const FVector Translation = Trans.GetTranslation();
+				const FVector Scale = Trans.GetScale3D();
+				return FString::Printf(TEXT("FTransform( FQuat(%f,%f,%f,%f), FVector(%f,%f,%f), FVector(%f,%f,%f) )"),
+					Rot.X, Rot.Y, Rot.Z, Rot.W, Translation.X, Translation.Y, Translation.Z, Scale.X, Scale.Y, Scale.Z);
+			}
+			else if (StructType == TBaseStructure<FLinearColor>::Get())
+			{
+				FLinearColor LinearColor;
+				LinearColor.InitFromString(CustomValue);
+				return FString::Printf(TEXT("FLinearColor(%f,%f,%f,%f)"), LinearColor.R, LinearColor.G, LinearColor.B, LinearColor.A);
+			}
+			else if (StructType == TBaseStructure<FColor>::Get())
+			{
+				FColor Color;
+				Color.InitFromString(CustomValue);
+				return FString::Printf(TEXT("FColor(%d,%d,%d,%d)"), Color.R, Color.G, Color.B, Color.A);
+			}
+			if (StructType == TBaseStructure<FVector2D>::Get())
+			{
+				FVector2D Vect = FVector2D::ZeroVector;
+				Vect.InitFromString(CustomValue);
+				return FString::Printf(TEXT("FVector2D(%f,%f)"), Vect.X, Vect.Y);
+			}
+			else
+			{
+				//@todo:  This needs to be more robust, since import text isn't really proper for struct construction.
+				return FString(TEXT("F")) + StructType->GetName() + (CustomValue.IsEmpty() ? TEXT("{}") : *CustomValue);
+			}
+		}
+		else if (Type.PinSubCategory == UEdGraphSchema_K2::PSC_Self)
+		{
+			return TEXT("this");
+		}
+		else if (UEdGraphSchema_K2::PC_Class == Type.PinCategory)
+		{
+			if (auto FoundClass = Cast<const UClass>(LiteralObject))
+			{
+				return FString::Printf(TEXT("%s%s::StaticClass()"), FoundClass->GetPrefixCPP(), *FoundClass->GetName());
+			}
+			return FString(TEXT("nullptr"));
+		}
+		else if((UEdGraphSchema_K2::PC_AssetClass == Type.PinCategory) || (UEdGraphSchema_K2::PC_Asset == Type.PinCategory))
+		{
+			return LiteralObject
+				? FString::Printf(TEXT("FStringAssetReference(TEXT(\"%s\"))"), *(LiteralObject->GetPathName()))
+				: FString(TEXT("nullptr"));
+		}
+		else if (UEdGraphSchema_K2::PC_Object == Type.PinCategory)
+		{
+			if (LiteralObject)
+			{
+				UClass* FoundClass = Cast<UClass>(Type.PinSubCategoryObject.Get());
+				FString ClassString = FoundClass ? (FString(FoundClass->GetPrefixCPP()) + FoundClass->GetName()) : TEXT("UObject");
+				return FString::Printf(TEXT("FindObject<%s>(ANY_PACKAGE, TEXT(\"%s\"))"), *ClassString, *(LiteralObject->GetPathName()));
+			}
+			else
+			{
+				return FString(TEXT("nullptr"));
+			}
+		}
+		/*
+		else if (CoerceProperty->IsA(UInterfaceProperty::StaticClass()))
+		{
+			if (Term->Type.PinSubCategory == PSC_Self)
+			{
+				return TEXT("this");
+			}
+			else
+			{
+				ensureMsgf(false, TEXT("It is not possible to express this interface property as a literal value!"));
+				return Term->Name;
+			}
+		}
+		*/
+		ensureMsgf(false, TEXT("It is not possible to express this type as a literal value!"));
+		return CustomValue;
+	}
+
+	static FString DefaultValue(const FEdGraphPinType& Type)
+	{
+		return LiteralTerm(Type, FString(), nullptr);
+	}
 };
 
diff --git a/Engine/Source/Editor/Kismet/Private/NativeCodeGenerationTool.cpp b/Engine/Source/Editor/Kismet/Private/NativeCodeGenerationTool.cpp
--- a/Engine/Source/Editor/Kismet/Private/NativeCodeGenerationTool.cpp
+++ b/Engine/Source/Editor/Kismet/Private/NativeCodeGenerationTool.cpp
@@ -57,6 +57,14 @@ struct FGeneratedCodeData
 			{
 				ReferencedObjects.Add(Class);
 			}
+
+			for (auto& ImplementedInterface : InBlueprint.GeneratedClass->Interfaces)
+			{
+				if (ImplementedInterface.Class && !ImplementedInterface.Class->HasAnyClassFlags(CLASS_Native))
+				{
+					ReferencedObjects.Add(ImplementedInterface.Class);
+				}
+			}
 		}
 
 		TypeDependencies.Empty();
diff --git a/Engine/Source/Editor/UnrealEd/Private/Kismet2/Kismet2.cpp b/Engine/Source/Editor/UnrealEd/Private/Kismet2/Kismet2.cpp
--- a/Engine/Source/Editor/UnrealEd/Private/Kismet2/Kismet2.cpp
+++ b/Engine/Source/Editor/UnrealEd/Private/Kismet2/Kismet2.cpp
@@ -975,6 +975,11 @@ void FKismetEditorUtilities::GenerateCppCode(UObject* Obj, TSharedPtr<FString> O
 				CompileOptions.OutCppSourceCode = OutCppSource;
 				CompileOptions.OutHeaderSourceCode = OutHeaderSource;
 				Compiler.CompileBlueprint(BlueprintObj, CompileOptions, Results);
+
+				if (EBlueprintType::BPTYPE_Interface == BlueprintObj->BlueprintType && OutCppSource.IsValid())
+				{
+					OutCppSource->Empty();
+				}
 			}
 			BlueprintObj->RemoveGeneratedClasses();
 			BlueprintObj->ClearFlags(RF_Standalone);
diff --git a/Engine/Source/Runtime/Engine/Public/GeneratedCodeHelpers.h b/Engine/Source/Runtime/Engine/Public/GeneratedCodeHelpers.h
--- a/Engine/Source/Runtime/Engine/Public/GeneratedCodeHelpers.h
+++ b/Engine/Source/Runtime/Engine/Public/GeneratedCodeHelpers.h
@@ -10,13 +10,34 @@
 // Special libraries
 #include "Kismet/DataTableFunctionLibrary.h"
 
+#include "Stack.h"
+
 FORCEINLINE UClass* DynamicMetaCast(const UClass* DesiredClass, UClass* SourceClass)
 {
 	return ((SourceClass)->IsChildOf(DesiredClass)) ? SourceClass : NULL;
 }
 
+FORCEINLINE bool IsValid(const FScriptInterface& Test)
+{
+	return IsValid(Test.GetObject()) && (nullptr != Test.GetInterface());
+}
+
 struct FCustomThunkTemplates
 {
+private:
+	static void ExecutionMessage(const TCHAR* Message, ELogVerbosity::Type Verbosity)
+	{
+		FFrame::KismetExecutionMessage(Message, Verbosity);
+	}
+
+	template<typename T>
+	static int32 LastIndexForLog(const TArray<T>& TargetArray)
+	{
+		const int32 ArraySize = TargetArray.Num();
+		return (ArraySize > 0) ? (ArraySize - 1) : 0;
+	}
+
+public:
 	//Replacements for CustomThunk functions from UKismetArrayLibrary
 
 	template<typename T>
@@ -54,16 +75,27 @@ struct FCustomThunkTemplates
 	template<typename T>
 	static void Array_Insert(TArray<T>& TargetArray, const T& NewItem, int32 Index)
 	{
-		TargetArray.Insert(NewItem, Index);
+		if ((Index >= 0) && (Index <= TargetArray.Num()))
+		{
+			TargetArray.Insert(NewItem, Index);
+		}
+		else
+		{
+			ExecutionMessage(*FString::Printf(TEXT("Attempted to insert an item into array out of bounds [%d/%d]!"), Index, LastIndexForLog(TargetArray)), ELogVerbosity::Warning);
+		}
 	}
 
 	template<typename T>
 	static void Array_Remove(TArray<T>& TargetArray, int32 IndexToRemove)
 	{
-		if (ensure(TargetArray.IsValidIndex(IndexToRemove)))
+		if (TargetArray.IsValidIndex(IndexToRemove))
 		{
 			TargetArray.RemoveAt(IndexToRemove);
 		}
+		else
+		{
+			ExecutionMessage(*FString::Printf(TEXT("Attempted to remove an item from an invalid index from array [%d/%d]!"), IndexToRemove, LastIndexForLog(TargetArray)), ELogVerbosity::Warning);
+		}
 	}
 
 	template<typename T>
@@ -87,7 +119,14 @@ struct FCustomThunkTemplates
 	template<typename T>
 	static void Array_Resize(TArray<T>& TargetArray, int32 Size)
 	{
-		TargetArray.SetNum(Size);
+		if (Size >= 0)
+		{
+			TargetArray.SetNum(Size);
+		}
+		else
+		{
+			ExecutionMessage(*FString::Printf(TEXT("Attempted to resize an array using negative size: Size = %d!"), Size), ELogVerbosity::Warning);
+		}
 	}
 
 	template<typename T>
@@ -105,23 +144,32 @@ struct FCustomThunkTemplates
 	template<typename T>
 	static void Array_Get(TArray<T>& TargetArray, int32 Index, T& Item)
 	{
-		Item = TargetArray[Index];
+		if (TargetArray.IsValidIndex(Index))
+		{
+			Item = TargetArray[Index];
+		}
+		else
+		{
+			ExecutionMessage(*FString::Printf(TEXT("Attempted to get an item from array out of bounds [%d/%d]!"), Index, LastIndexForLog(TargetArray)), ELogVerbosity::Warning);
+			Item = T{};
+		}
 	}
 
 	template<typename T>
 	static void Array_Set(TArray<T>& TargetArray, int32 Index, const T& Item, bool bSizeToFit)
 	{
-		if (ensure(Index >= 0))
+		if (!TargetArray.IsValidIndex(Index) && bSizeToFit && (Index >= 0))
 		{
-			if (!TargetArray.IsValidIndex(Index) && bSizeToFit)
-			{
-				TargetArray.SetNum(Index + 1);
-			}
+			TargetArray.SetNum(Index + 1);
+		}
 
-			if (TargetArray.IsValidIndex(Index))
-			{
-				TargetArray[Index] = Item;
-			}
+		if (TargetArray.IsValidIndex(Index))
+		{
+			TargetArray[Index] = Item;
+		}
+		else
+		{
+			ExecutionMessage(*FString::Printf(TEXT("Attempted to set an invalid index on array [%d/%d]!"), Index, LastIndexForLog(TargetArray)), ELogVerbosity::Warning);
 		}
 	}
 
