UEBP-40 C++ code generation (WIP):
- interface is called with "::Execute_"
- interface functions are BlueprintNativeEvents
- CppFromBpEvent metadata added, to force FUNC_Event
- BlueprintNativeEvent support
- BlueprintImplementableEvent support

diff --git a/Engine/Source/Developer/BlueprintCompilerCppBackend/Private/BlueprintCompilerCppBackend.cpp b/Engine/Source/Developer/BlueprintCompilerCppBackend/Private/BlueprintCompilerCppBackend.cpp
--- a/Engine/Source/Developer/BlueprintCompilerCppBackend/Private/BlueprintCompilerCppBackend.cpp
+++ b/Engine/Source/Developer/BlueprintCompilerCppBackend/Private/BlueprintCompilerCppBackend.cpp
@@ -291,6 +291,7 @@ FString FBlueprintCompilerCppBackend::EmitCallStatmentInner(FBlueprintCompiledSt
 	const bool bCallOnDifferentObject = Statement.FunctionContext && (Statement.FunctionContext->Name != TEXT("self"));
 	const bool bStaticCall = Statement.FunctionToCall->HasAnyFunctionFlags(FUNC_Static);
 	const bool bUseSafeContext = bCallOnDifferentObject && !bStaticCall;
+	const bool bInterfaceCall = bCallOnDifferentObject && Statement.FunctionContext && (UEdGraphSchema_K2::PC_Interface == Statement.FunctionContext->Type.PinCategory);
 	{
 		FSafeContextScopedEmmitter SafeContextScope(Result, bUseSafeContext ? Statement.FunctionContext : nullptr, *this, TEXT("\t\t\t"));
 		ensure(!bInline || !SafeContextScope.IsSafeContextUsed());
@@ -309,22 +310,42 @@ FString FBlueprintCompilerCppBackend::EmitCallStatmentInner(FBlueprintCompiledSt
 		}
 
 		// Emit object to call the method on
-		if (bStaticCall)
+		if (bInterfaceCall)
 		{
-			const bool bIsCustomThunk = Statement.FunctionToCall->HasMetaData(TEXT("CustomStructureParam")) || Statement.FunctionToCall->HasMetaData(TEXT("ArrayParm"));
-			auto OwnerClass = Statement.FunctionToCall->GetOuterUClass();
-			Result += bIsCustomThunk ? TEXT("FCustomThunkTemplates::") : FString::Printf(TEXT("%s%s::"), OwnerClass->GetPrefixCPP(), *OwnerClass->GetName());
+			auto ContextInterfaceClass = CastChecked<UClass>(Statement.FunctionContext->Type.PinSubCategoryObject.Get());
+			ensure(ContextInterfaceClass->IsChildOf<UInterface>());
+			Result += FString::Printf(TEXT("I%s::Execute_%s(%s.GetObject(), ")
+				, *ContextInterfaceClass->GetName()
+				, *Statement.FunctionToCall->GetName()
+				, *TermToText(Statement.FunctionContext, false));
 		}
-		else if (bCallOnDifferentObject) //@TODO: Badness, could be a self reference wired to another instance!
+		else
 		{
-			Result += FString::Printf(TEXT("%s->"), *TermToText(Statement.FunctionContext, false));
-		}
+			if (bStaticCall)
+			{
+				const bool bIsCustomThunk = Statement.FunctionToCall->HasMetaData(TEXT("CustomStructureParam")) || Statement.FunctionToCall->HasMetaData(TEXT("ArrayParm"));
+				auto OwnerClass = Statement.FunctionToCall->GetOuterUClass();
+				Result += bIsCustomThunk ? TEXT("FCustomThunkTemplates::") : FString::Printf(TEXT("%s%s::"), OwnerClass->GetPrefixCPP(), *OwnerClass->GetName());
+			}
+			else if (bCallOnDifferentObject) //@TODO: Badness, could be a self reference wired to another instance!
+			{
+				Result += FString::Printf(TEXT("%s->"), *TermToText(Statement.FunctionContext, false));
+			}
 
-		// Emit method name
-		Result += FString::Printf(TEXT("%s%s"), Statement.bIsParentContext ? TEXT("Super::") : TEXT(""), *Statement.FunctionToCall->GetName());
+			if (Statement.bIsParentContext)
+			{
+				Result += TEXT("Super::");
+			}
+			Result += Statement.FunctionToCall->GetName();
+			if (Statement.bIsParentContext && FEmitHelper::ShoulsHandleAsNativeEvent(Statement.FunctionToCall))
+			{
+				ensure(!bCallOnDifferentObject);
+				Result += TEXT("_Implementation");
+			}
 
-		// Emit method parameter list
-		Result += TEXT("(");
+			// Emit method parameter list
+			Result += TEXT("(");
+		}
 		Result += EmitMethodInputParameterList(Statement);
 		Result += TEXT(")");
 
@@ -702,6 +723,19 @@ void FBlueprintCompilerCppBackend::ConstructFunction(FKismetFunctionContext& Fun
 			}
 		}
 
+		bool bAddCppFromBpEventMD = false;
+		bool bGenerateAsNativeEventImplementation = false;
+		// Get original function declaration
+		if (FEmitHelper::ShoulsHandleAsNativeEvent(Function)) // BlueprintNativeEvent
+		{
+			bGenerateAsNativeEventImplementation = true;
+			FunctionName += TEXT("_Implementation");
+		}
+		else if (FEmitHelper::ShoulsHandleAsImplementableEvent(Function)) // BlueprintImplementableEvent
+		{
+			bAddCppFromBpEventMD = true;
+		}
+
 		// Emit the declaration
 		const FString ReturnType = ReturnValue ? ReturnValue->GetCPPType(NULL, EPropertyExportCPPFlags::CPPF_CustomTypeName) : TEXT("void");
 
@@ -709,12 +743,28 @@ void FBlueprintCompilerCppBackend::ConstructFunction(FKismetFunctionContext& Fun
 		{
 			const FString Start = FString::Printf(TEXT("%s %s%s%s("), *ReturnType, TEXT("%s"), TEXT("%s"), *FunctionName);
 
-			Emit(Header, *FString::Printf(TEXT("\t%s\n"), *FEmitHelper::EmitUFuntion(Function)));
+			TArray<FString> AdditionalMetaData;
+			if (bAddCppFromBpEventMD)
+			{
+				AdditionalMetaData.Emplace(TEXT("CppFromBpEvent"));
+			}
+
+			if (!bGenerateAsNativeEventImplementation)
+			{
+				Emit(Header, *FString::Printf(TEXT("\t%s\n"), *FEmitHelper::EmitUFuntion(Function, &AdditionalMetaData)));
+			}
 			Emit(Header, TEXT("\t"));
+
 			if (Function->HasAllFunctionFlags(FUNC_Static))
 			{
 				Emit(Header, TEXT("static "));
 			}
+
+			if (bGenerateAsNativeEventImplementation)
+			{
+				Emit(Header, TEXT("virtual "));
+			}
+
 			Emit(Header, *FString::Printf(*Start, TEXT(""), TEXT("")));
 			Emit(Body, *FString::Printf(*Start, *CppClassName, TEXT("::")));
 
@@ -739,6 +789,10 @@ void FBlueprintCompilerCppBackend::ConstructFunction(FKismetFunctionContext& Fun
 			}
 
 			Emit(Header, TEXT(")"));
+			if (bGenerateAsNativeEventImplementation)
+			{
+				Emit(Header, TEXT(" override"));
+			}
 			Emit(Header, TEXT(";\n"));
 			Emit(Body, TEXT(")\n"));
 		}
diff --git a/Engine/Source/Developer/BlueprintCompilerCppBackend/Private/BlueprintCompilerCppBackendUtils.h b/Engine/Source/Developer/BlueprintCompilerCppBackend/Private/BlueprintCompilerCppBackendUtils.h
--- a/Engine/Source/Developer/BlueprintCompilerCppBackend/Private/BlueprintCompilerCppBackendUtils.h
+++ b/Engine/Source/Developer/BlueprintCompilerCppBackend/Private/BlueprintCompilerCppBackendUtils.h
@@ -117,7 +117,7 @@ struct FEmitHelper
 		return true;
 	}
 
-	static FString HandleMetaData(const UField* Field, bool AddCategory = true)
+	static FString HandleMetaData(const UField* Field, bool AddCategory = true, TArray<FString>* AdditinalMetaData = nullptr)
 	{
 		FString MetaDataStr;
 
@@ -151,6 +151,10 @@ struct FEmitHelper
 		{
 			MetaDataStrings.Emplace(TEXT("Category"));
 		}
+		if (AdditinalMetaData)
+		{
+			MetaDataStrings.Append(*AdditinalMetaData);
+		}
 		MetaDataStrings.Remove(FString());
 		if (MetaDataStrings.Num())
 		{
@@ -218,27 +222,27 @@ struct FEmitHelper
 
 		// OTHER
 		HANDLE_CPF_TAG(TEXT("Transient"), CPF_Transient)
-			HANDLE_CPF_TAG(TEXT("DuplicateTransient"), CPF_DuplicateTransient)
-			HANDLE_CPF_TAG(TEXT("TextExportTransient"), CPF_TextExportTransient)
-			HANDLE_CPF_TAG(TEXT("NonPIEDuplicateTransient"), CPF_NonPIEDuplicateTransient)
-			HANDLE_CPF_TAG(TEXT("Export"), CPF_ExportObject)
-			HANDLE_CPF_TAG(TEXT("NoClear"), CPF_NoClear)
-			HANDLE_CPF_TAG(TEXT("EditFixedSize"), CPF_EditFixedSize)
-			HANDLE_CPF_TAG(TEXT("NotReplicated"), CPF_RepSkip)
-			HANDLE_CPF_TAG(TEXT("Interp"), CPF_Edit | CPF_BlueprintVisible | CPF_Interp)
-			HANDLE_CPF_TAG(TEXT("NonTransactional"), CPF_NonTransactional)
-			HANDLE_CPF_TAG(TEXT("BlueprintAssignable"), CPF_BlueprintAssignable)
-			HANDLE_CPF_TAG(TEXT("BlueprintCallable"), CPF_BlueprintCallable)
-			HANDLE_CPF_TAG(TEXT("BlueprintAuthorityOnly"), CPF_BlueprintAuthorityOnly)
-			HANDLE_CPF_TAG(TEXT("AssetRegistrySearchable"), CPF_AssetRegistrySearchable)
-			HANDLE_CPF_TAG(TEXT("SimpleDisplay"), CPF_SimpleDisplay)
-			HANDLE_CPF_TAG(TEXT("AdvancedDisplay"), CPF_AdvancedDisplay)
-			HANDLE_CPF_TAG(TEXT("SaveGame"), CPF_SaveGame)
-
-			//TODO:
-			//HANDLE_CPF_TAG(TEXT("Instanced"), CPF_PersistentInstance | CPF_ExportObject | CPF_InstancedReference)
-
-			return Tags;
+		HANDLE_CPF_TAG(TEXT("DuplicateTransient"), CPF_DuplicateTransient)
+		HANDLE_CPF_TAG(TEXT("TextExportTransient"), CPF_TextExportTransient)
+		HANDLE_CPF_TAG(TEXT("NonPIEDuplicateTransient"), CPF_NonPIEDuplicateTransient)
+		HANDLE_CPF_TAG(TEXT("Export"), CPF_ExportObject)
+		HANDLE_CPF_TAG(TEXT("NoClear"), CPF_NoClear)
+		HANDLE_CPF_TAG(TEXT("EditFixedSize"), CPF_EditFixedSize)
+		HANDLE_CPF_TAG(TEXT("NotReplicated"), CPF_RepSkip)
+		HANDLE_CPF_TAG(TEXT("Interp"), CPF_Edit | CPF_BlueprintVisible | CPF_Interp)
+		HANDLE_CPF_TAG(TEXT("NonTransactional"), CPF_NonTransactional)
+		HANDLE_CPF_TAG(TEXT("BlueprintAssignable"), CPF_BlueprintAssignable)
+		HANDLE_CPF_TAG(TEXT("BlueprintCallable"), CPF_BlueprintCallable)
+		HANDLE_CPF_TAG(TEXT("BlueprintAuthorityOnly"), CPF_BlueprintAuthorityOnly)
+		HANDLE_CPF_TAG(TEXT("AssetRegistrySearchable"), CPF_AssetRegistrySearchable)
+		HANDLE_CPF_TAG(TEXT("SimpleDisplay"), CPF_SimpleDisplay)
+		HANDLE_CPF_TAG(TEXT("AdvancedDisplay"), CPF_AdvancedDisplay)
+		HANDLE_CPF_TAG(TEXT("SaveGame"), CPF_SaveGame)
+
+		//TODO:
+		//HANDLE_CPF_TAG(TEXT("Instanced"), CPF_PersistentInstance | CPF_ExportObject | CPF_InstancedReference)
+
+		return Tags;
 	}
 
 	static TArray<FString> FunctionFlagsToTags(uint64 Flags)
@@ -246,24 +250,24 @@ struct FEmitHelper
 		TArray<FString> Tags;
 
 		//Pointless: BlueprintNativeEvent, BlueprintImplementableEvent
-		//Pointless: CustomThunk ?
+		//Pointless: CustomThunk
 
 		//TODO: SealedEvent
 		//TODO: Unreliable
 		//TODO: ServiceRequest, ServiceResponse
 
 		HANDLE_CPF_TAG(TEXT("Exec"), FUNC_Exec)
-			HANDLE_CPF_TAG(TEXT("Server"), FUNC_Net | FUNC_NetServer)
-			HANDLE_CPF_TAG(TEXT("Client"), FUNC_Net | FUNC_NetClient)
-			HANDLE_CPF_TAG(TEXT("NetMulticast"), FUNC_Net | FUNC_NetMulticast)
-			HANDLE_CPF_TAG(TEXT("Reliable"), FUNC_NetReliable)
-			HANDLE_CPF_TAG(TEXT("BlueprintCallable"), FUNC_BlueprintCallable)
-			HANDLE_CPF_TAG(TEXT("BlueprintPure"), FUNC_BlueprintCallable | FUNC_BlueprintPure)
-			HANDLE_CPF_TAG(TEXT("BlueprintAuthorityOnly"), FUNC_BlueprintAuthorityOnly)
-			HANDLE_CPF_TAG(TEXT("BlueprintCosmetic"), FUNC_BlueprintCosmetic)
-			HANDLE_CPF_TAG(TEXT("WithValidation"), FUNC_NetValidate)
-
-			return Tags;
+		HANDLE_CPF_TAG(TEXT("Server"), FUNC_Net | FUNC_NetServer)
+		HANDLE_CPF_TAG(TEXT("Client"), FUNC_Net | FUNC_NetClient)
+		HANDLE_CPF_TAG(TEXT("NetMulticast"), FUNC_Net | FUNC_NetMulticast)
+		HANDLE_CPF_TAG(TEXT("Reliable"), FUNC_NetReliable)
+		HANDLE_CPF_TAG(TEXT("BlueprintCallable"), FUNC_BlueprintCallable)
+		HANDLE_CPF_TAG(TEXT("BlueprintPure"), FUNC_BlueprintCallable | FUNC_BlueprintPure)
+		HANDLE_CPF_TAG(TEXT("BlueprintAuthorityOnly"), FUNC_BlueprintAuthorityOnly)
+		HANDLE_CPF_TAG(TEXT("BlueprintCosmetic"), FUNC_BlueprintCosmetic)
+		HANDLE_CPF_TAG(TEXT("WithValidation"), FUNC_NetValidate)
+
+		return Tags;
 	}
 
 #undef HANDLE_CPF_TAG
@@ -278,11 +282,18 @@ struct FEmitHelper
 		return HasAllFlags(FunctionFlags, FUNC_Event | FUNC_BlueprintEvent) && !HasAllFlags(FunctionFlags, FUNC_Native);
 	}
 
-	static FString EmitUFuntion(UFunction* Function)
+	static FString EmitUFuntion(UFunction* Function, TArray<FString>* AdditinalMetaData = nullptr)
 	{
 		TArray<FString> Tags = FEmitHelper::FunctionFlagsToTags(Function->FunctionFlags);
+
+		auto FunctionOwnerClass = Function->GetOwnerClass();
+		if (FunctionOwnerClass->IsChildOf<UInterface>())
+		{
+			Tags.Emplace(TEXT("BlueprintNativeEvent"));
+		}
+
 		const bool bMustHaveCategory = (Function->FunctionFlags & (FUNC_BlueprintCallable | FUNC_BlueprintPure)) != 0;
-		Tags.Emplace(FEmitHelper::HandleMetaData(Function, bMustHaveCategory));
+		Tags.Emplace(FEmitHelper::HandleMetaData(Function, bMustHaveCategory, AdditinalMetaData));
 		Tags.Remove(FString());
 
 		FString AllTags;
@@ -583,5 +594,61 @@ struct FEmitHelper
 	{
 		return LiteralTerm(Type, FString(), nullptr);
 	}
-};
 
+	static UFunction* GetOriginalFunction(UFunction* Function)
+	{
+		check(Function);
+		const FName FunctionName = Function->GetFName();
+
+		UClass* Owner = Function->GetOwnerClass();
+		check(Owner);
+		for (auto& Inter : Owner->Interfaces)
+		{
+			if (UFunction* Result = Inter.Class->FindFunctionByName(FunctionName))
+			{
+				return GetOriginalFunction(Result);
+			}
+		}
+
+		for (UClass* SearchClass = Owner->GetSuperClass(); SearchClass; SearchClass = SearchClass->GetSuperClass())
+		{
+			if (UFunction* Result = SearchClass->FindFunctionByName(FunctionName))
+			{
+				return GetOriginalFunction(Result);
+			}
+		}
+
+		return Function;
+	}
+
+	static bool ShoulsHandleAsNativeEvent(UFunction* Function)
+	{
+		check(Function);
+		auto OriginalFunction = FEmitHelper::GetOriginalFunction(Function);
+		check(OriginalFunction);
+		if (OriginalFunction != Function)
+		{
+			const uint32 FlagsToCheckMask = EFunctionFlags::FUNC_Event | EFunctionFlags::FUNC_BlueprintEvent | EFunctionFlags::FUNC_Native;
+			const uint32 FlagsToCheck = OriginalFunction->FunctionFlags & FlagsToCheckMask;
+
+			auto OriginalOwner = OriginalFunction->GetOwnerClass();
+			const bool bBPInterface = Cast<UBlueprintGeneratedClass>(OriginalOwner) && OriginalOwner->IsChildOf<UInterface>();
+			return (FlagsToCheck == FlagsToCheckMask) || bBPInterface;
+		}
+		return false;
+	}
+
+	static bool ShoulsHandleAsImplementableEvent(UFunction* Function)
+	{
+		check(Function);
+		auto OriginalFunction = FEmitHelper::GetOriginalFunction(Function);
+		check(OriginalFunction);
+		if (OriginalFunction != Function)
+		{
+			const uint32 FlagsToCheckMask = EFunctionFlags::FUNC_Event | EFunctionFlags::FUNC_BlueprintEvent | EFunctionFlags::FUNC_Native;
+			const uint32 FlagsToCheck = OriginalFunction->FunctionFlags & FlagsToCheckMask;
+			return (FlagsToCheck == (EFunctionFlags::FUNC_Event | EFunctionFlags::FUNC_BlueprintEvent));
+		}
+		return false;
+	}
+};
\ No newline at end of file
diff --git a/Engine/Source/Programs/UnrealHeaderTool/Private/HeaderParser.cpp b/Engine/Source/Programs/UnrealHeaderTool/Private/HeaderParser.cpp
--- a/Engine/Source/Programs/UnrealHeaderTool/Private/HeaderParser.cpp
+++ b/Engine/Source/Programs/UnrealHeaderTool/Private/HeaderParser.cpp
@@ -5295,6 +5295,11 @@ void FHeaderParser::CompileFunctionDeclaration(FUnrealSourceFile& SourceFile, FC
 		FuncInfo.FunctionExportFlags |= FUNCEXPORT_CppStatic;
 	}
 
+	if (MetaData.Contains("CppFromBpEvent"))
+	{
+		FuncInfo.FunctionFlags |= FUNC_Event;
+	}
+
 	ProcessFunctionSpecifiers(FuncInfo, SpecifiersFound);
 
 	if ((FuncInfo.FunctionFlags & FUNC_BlueprintPure) && GetCurrentClass()->HasAnyClassFlags(CLASS_Interface))
