UEBP-40 C++ code generation (WIP):
- CppBackend: A prototype of the navite code generating - to create default values in a struct.
- Added PPF_ExportCpp flag. Work In Progress.

diff --git a/Engine/Source/Developer/BlueprintCompilerCppBackend/Private/BlueprintCompilerCppBackend.cpp b/Engine/Source/Developer/BlueprintCompilerCppBackend/Private/BlueprintCompilerCppBackend.cpp
--- a/Engine/Source/Developer/BlueprintCompilerCppBackend/Private/BlueprintCompilerCppBackend.cpp
+++ b/Engine/Source/Developer/BlueprintCompilerCppBackend/Private/BlueprintCompilerCppBackend.cpp
@@ -3,6 +3,7 @@
 #include "BlueprintCompilerCppBackendModulePrivatePCH.h"
 #include "BlueprintCompilerCppBackend.h"
 #include "BlueprintCompilerCppBackendUtils.h"
+#include "EdGraphSchema_K2.h"
 
 FString FBlueprintCompilerCppBackend::TermToText(const FBPTerminal* Term, bool bUseSafeContext)
 {
@@ -1016,6 +1017,7 @@ void FBlueprintCompilerCppBackend::GenerateCodeFromStruct(UUserDefinedStruct* So
 
 	Emit(Header, *FString::Printf(TEXT("struct %s\n{\npublic:\n\tGENERATED_BODY()\n"), *NewName));
 	EmitStructProperties(Header, SourceStruct);
+	Emit(Header, *FEmitDefaultValueHelper::GenerateGetDefaultValue(SourceStruct));
 	Emit(Header, TEXT("};\n"));
 }
 
diff --git a/Engine/Source/Developer/BlueprintCompilerCppBackend/Private/BlueprintCompilerCppBackendUtils.h b/Engine/Source/Developer/BlueprintCompilerCppBackend/Private/BlueprintCompilerCppBackendUtils.h
--- a/Engine/Source/Developer/BlueprintCompilerCppBackend/Private/BlueprintCompilerCppBackendUtils.h
+++ b/Engine/Source/Developer/BlueprintCompilerCppBackend/Private/BlueprintCompilerCppBackendUtils.h
@@ -1,6 +1,7 @@
 // Copyright 1998-2015 Epic Games, Inc. All Rights Reserved.
+#pragma  once
 
-#include "EdGraphSchema_K2.h"
+class FBlueprintCompilerCppBackend;
 
 // Generates single if scope. It's condition checks context of given term.
 struct FSafeContextScopedEmmitter
@@ -8,647 +9,62 @@ struct FSafeContextScopedEmmitter
 private:
 	FString& Body;
 	bool bSafeContextUsed;
-	const TCHAR* CurrentIndent;
+	FString CurrentIndent;
 public:
-	FString GetAdditionalIndent() const
-	{
-		return bSafeContextUsed ? TEXT("\t") : FString();
-	}
+	FString GetAdditionalIndent() const;
+	bool IsSafeContextUsed() const;
 
-	bool IsSafeContextUsed() const
-	{
-		return bSafeContextUsed;
-	}
+	static FString ValidationChain(const FBPTerminal* Term, FBlueprintCompilerCppBackend& CppBackend);
 
-	static FString ValidationChain(const FBPTerminal* Term, FBlueprintCompilerCppBackend& CppBackend)
-	{
-		TArray<FString> SafetyConditions;
-		for (; Term; Term = Term->Context)
-		{
-			if (!Term->IsStructContextType() && (Term->Type.PinSubCategory != TEXT("self")))
-			{
-				ensure(!Term->bIsLiteral);
-				SafetyConditions.Add(CppBackend.TermToText(Term, false));
-			}
-		}
-
-		FString Result;
-		for (int32 Iter = SafetyConditions.Num() - 1; Iter >= 0; --Iter)
-		{
-			Result += FString(TEXT("IsValid("));
-			Result += SafetyConditions[Iter];
-			Result += FString(TEXT(")"));
-			if (Iter)
-			{
-				Result += FString(TEXT(" && "));
-			}
-		}
-
-		return Result;
-	}
-
-	FSafeContextScopedEmmitter(FString& InBody, const FBPTerminal* Term, FBlueprintCompilerCppBackend& CppBackend, const TCHAR* InCurrentIndent)
-		: Body(InBody), bSafeContextUsed(false), CurrentIndent(InCurrentIndent)
-	{
-		const FString Conditions = ValidationChain(Term, CppBackend);
-
-		if (!Conditions.IsEmpty())
-		{
-			bSafeContextUsed = true;
-			Body += CurrentIndent;
-			Body += TEXT("if (");
-			Body += Conditions;
-			Body += TEXT(")\n");
-			Body += CurrentIndent;
-			Body += TEXT("{\n");
-		}
-	}
-
-	~FSafeContextScopedEmmitter()
-	{
-		if (bSafeContextUsed)
-		{
-			Body += CurrentIndent;
-			Body += TEXT("}\n");
-		}
-	}
+	FSafeContextScopedEmmitter(FString& InBody, const FBPTerminal* Term, FBlueprintCompilerCppBackend& CppBackend, const TCHAR* InCurrentIndent);
+	~FSafeContextScopedEmmitter();
 };
 
 struct FEmitHelper
 {
-	static void ArrayToString(const TArray<FString>& Array, FString& String, const TCHAR* Separator)
-	{
-		if (Array.Num())
-		{
-			String += Array[0];
-		}
-		for (int32 i = 1; i < Array.Num(); ++i)
-		{
-			String += TEXT(", ");
-			String += Array[i];
-		}
-	}
-
-	static bool HasAllFlags(uint64 Flags, uint64 FlagsToCheck)
-	{
-		return FlagsToCheck == (Flags & FlagsToCheck);
-	}
-
-	static FString HandleRepNotifyFunc(const UProperty* Property)
-	{
-		check(Property);
-		if (HasAllFlags(Property->PropertyFlags, CPF_Net | CPF_RepNotify))
-		{
-			return FString::Printf(TEXT("ReplicatedUsing=%s"), *Property->RepNotifyFunc.ToString());
-		}
-		else if (HasAllFlags(Property->PropertyFlags, CPF_Net))
-		{
-			return TEXT("Replicated");
-		}
-		return FString();
-	}
-
-	static bool MetaDataCanBeNative(const FName MetaDataName)
-	{
-		if (MetaDataName == TEXT("ModuleRelativePath"))
-		{
-			return false;
-		}
-		return true;
-	}
-
-	static FString HandleMetaData(const UField* Field, bool AddCategory = true, TArray<FString>* AdditinalMetaData = nullptr)
-	{
-		FString MetaDataStr;
-
-		check(Field);
-		UPackage* Package = Field->GetOutermost();
-		check(Package);
-		const UMetaData* MetaData = Package->GetMetaData();
-		check(MetaData);
-		const TMap<FName, FString>* ValuesMap = MetaData->ObjectMetaDataMap.Find(Field);
-		TArray<FString> MetaDataStrings;
-		if (ValuesMap && ValuesMap->Num())
-		{
-			for (auto& Pair : *ValuesMap)
-			{
-				if (!MetaDataCanBeNative(Pair.Key))
-				{
-					continue;
-				}
-				if (!Pair.Value.IsEmpty())
-				{
-					FString Value = Pair.Value.Replace(TEXT("\n"), TEXT(""));
-					MetaDataStrings.Emplace(FString::Printf(TEXT("%s=\"%s\""), *Pair.Key.ToString(), *Value));
-				}
-				else
-				{
-					MetaDataStrings.Emplace(Pair.Key.ToString());
-				}
-			}
-		}
-		if (AddCategory && (!ValuesMap || !ValuesMap->Find(TEXT("Category"))))
-		{
-			MetaDataStrings.Emplace(TEXT("Category"));
-		}
-		if (AdditinalMetaData)
-		{
-			MetaDataStrings.Append(*AdditinalMetaData);
-		}
-		MetaDataStrings.Remove(FString());
-		if (MetaDataStrings.Num())
-		{
-			MetaDataStr += TEXT("meta=(");
-			ArrayToString(MetaDataStrings, MetaDataStr, TEXT(", "));
-			MetaDataStr += TEXT(")");
-		}
-		return MetaDataStr;
-	}
-
-#ifdef HANDLE_CPF_TAG
-	static_assert(false, "Macro HANDLE_CPF_TAG redefinition.");
-#endif
-#define HANDLE_CPF_TAG(TagName, CheckedFlags) if (HasAllFlags(Flags, (CheckedFlags))) { Tags.Emplace(TagName); }
-
-	static TArray<FString> ProperyFlagsToTags(uint64 Flags)
-	{
-		TArray<FString> Tags;
-
-		// EDIT FLAGS
-		if (HasAllFlags(Flags, (CPF_Edit | CPF_EditConst | CPF_DisableEditOnInstance)))
-		{
-			Tags.Emplace(TEXT("VisibleDefaultsOnly"));
-		}
-		else if (HasAllFlags(Flags, (CPF_Edit | CPF_EditConst | CPF_DisableEditOnTemplate)))
-		{
-			Tags.Emplace(TEXT("VisibleInstanceOnly"));
-		}
-		else if (HasAllFlags(Flags, (CPF_Edit | CPF_EditConst)))
-		{
-			Tags.Emplace(TEXT("VisibleAnywhere"));
-		}
-		else if (HasAllFlags(Flags, (CPF_Edit | CPF_DisableEditOnInstance)))
-		{
-			Tags.Emplace(TEXT("EditDefaultsOnly"));
-		}
-		else if (HasAllFlags(Flags, (CPF_Edit | CPF_DisableEditOnTemplate)))
-		{
-			Tags.Emplace(TEXT("EditInstanceOnly"));
-		}
-		else if (HasAllFlags(Flags, (CPF_Edit)))
-		{
-			Tags.Emplace(TEXT("EditAnywhere"));
-		}
-
-		// BLUEPRINT EDIT
-		if (HasAllFlags(Flags, (CPF_BlueprintVisible | CPF_BlueprintReadOnly)))
-		{
-			Tags.Emplace(TEXT("BlueprintReadOnly"));
-		}
-		else if (HasAllFlags(Flags, (CPF_BlueprintVisible)))
-		{
-			Tags.Emplace(TEXT("BlueprintReadWrite"));
-		}
-
-		// CONFIG
-		if (HasAllFlags(Flags, (CPF_GlobalConfig | CPF_Config)))
-		{
-			Tags.Emplace(TEXT("GlobalConfig"));
-		}
-		else if (HasAllFlags(Flags, (CPF_Config)))
-		{
-			Tags.Emplace(TEXT("Config"));
-		}
-
-		// OTHER
-		HANDLE_CPF_TAG(TEXT("Transient"), CPF_Transient)
-		HANDLE_CPF_TAG(TEXT("DuplicateTransient"), CPF_DuplicateTransient)
-		HANDLE_CPF_TAG(TEXT("TextExportTransient"), CPF_TextExportTransient)
-		HANDLE_CPF_TAG(TEXT("NonPIEDuplicateTransient"), CPF_NonPIEDuplicateTransient)
-		HANDLE_CPF_TAG(TEXT("Export"), CPF_ExportObject)
-		HANDLE_CPF_TAG(TEXT("NoClear"), CPF_NoClear)
-		HANDLE_CPF_TAG(TEXT("EditFixedSize"), CPF_EditFixedSize)
-		HANDLE_CPF_TAG(TEXT("NotReplicated"), CPF_RepSkip)
-		HANDLE_CPF_TAG(TEXT("Interp"), CPF_Edit | CPF_BlueprintVisible | CPF_Interp)
-		HANDLE_CPF_TAG(TEXT("NonTransactional"), CPF_NonTransactional)
-		HANDLE_CPF_TAG(TEXT("BlueprintAssignable"), CPF_BlueprintAssignable)
-		HANDLE_CPF_TAG(TEXT("BlueprintCallable"), CPF_BlueprintCallable)
-		HANDLE_CPF_TAG(TEXT("BlueprintAuthorityOnly"), CPF_BlueprintAuthorityOnly)
-		HANDLE_CPF_TAG(TEXT("AssetRegistrySearchable"), CPF_AssetRegistrySearchable)
-		HANDLE_CPF_TAG(TEXT("SimpleDisplay"), CPF_SimpleDisplay)
-		HANDLE_CPF_TAG(TEXT("AdvancedDisplay"), CPF_AdvancedDisplay)
-		HANDLE_CPF_TAG(TEXT("SaveGame"), CPF_SaveGame)
-
-		//TODO:
-		//HANDLE_CPF_TAG(TEXT("Instanced"), CPF_PersistentInstance | CPF_ExportObject | CPF_InstancedReference)
-
-		return Tags;
-	}
-
-	static TArray<FString> FunctionFlagsToTags(uint64 Flags)
-	{
-		TArray<FString> Tags;
+	static void ArrayToString(const TArray<FString>& Array, FString& String, const TCHAR* Separator);
 
-		//Pointless: BlueprintNativeEvent, BlueprintImplementableEvent
-		//Pointless: CustomThunk
+	static bool HasAllFlags(uint64 Flags, uint64 FlagsToCheck);
 
-		//TODO: SealedEvent
-		//TODO: Unreliable
-		//TODO: ServiceRequest, ServiceResponse
+	static FString HandleRepNotifyFunc(const UProperty* Property);
 
-		HANDLE_CPF_TAG(TEXT("Exec"), FUNC_Exec)
-		HANDLE_CPF_TAG(TEXT("Server"), FUNC_Net | FUNC_NetServer)
-		HANDLE_CPF_TAG(TEXT("Client"), FUNC_Net | FUNC_NetClient)
-		HANDLE_CPF_TAG(TEXT("NetMulticast"), FUNC_Net | FUNC_NetMulticast)
-		HANDLE_CPF_TAG(TEXT("Reliable"), FUNC_NetReliable)
-		HANDLE_CPF_TAG(TEXT("BlueprintCallable"), FUNC_BlueprintCallable)
-		HANDLE_CPF_TAG(TEXT("BlueprintPure"), FUNC_BlueprintCallable | FUNC_BlueprintPure)
-		HANDLE_CPF_TAG(TEXT("BlueprintAuthorityOnly"), FUNC_BlueprintAuthorityOnly)
-		HANDLE_CPF_TAG(TEXT("BlueprintCosmetic"), FUNC_BlueprintCosmetic)
-		HANDLE_CPF_TAG(TEXT("WithValidation"), FUNC_NetValidate)
+	static bool MetaDataCanBeNative(const FName MetaDataName);
 
-		return Tags;
-	}
+	static FString HandleMetaData(const UField* Field, bool AddCategory = true, TArray<FString>* AdditinalMetaData = nullptr);
 
-#undef HANDLE_CPF_TAG
+	static TArray<FString> ProperyFlagsToTags(uint64 Flags);
 
-	static bool IsBlueprintNativeEvent(uint64 FunctionFlags)
-	{
-		return HasAllFlags(FunctionFlags, FUNC_Event | FUNC_BlueprintEvent | FUNC_Native);
-	}
+	static TArray<FString> FunctionFlagsToTags(uint64 Flags);
 
-	static bool IsBlueprintImplementableEvent(uint64 FunctionFlags)
-	{
-		return HasAllFlags(FunctionFlags, FUNC_Event | FUNC_BlueprintEvent) && !HasAllFlags(FunctionFlags, FUNC_Native);
-	}
+	static bool IsBlueprintNativeEvent(uint64 FunctionFlags);
 
-	static FString EmitUFuntion(UFunction* Function, TArray<FString>* AdditinalMetaData = nullptr)
-	{
-		TArray<FString> Tags = FEmitHelper::FunctionFlagsToTags(Function->FunctionFlags);
+	static bool IsBlueprintImplementableEvent(uint64 FunctionFlags);
 
-		auto FunctionOwnerClass = Function->GetOwnerClass();
-		if (FunctionOwnerClass->IsChildOf<UInterface>())
-		{
-			Tags.Emplace(TEXT("BlueprintNativeEvent"));
-		}
+	static FString EmitUFuntion(UFunction* Function, TArray<FString>* AdditinalMetaData = nullptr);
 
-		const bool bMustHaveCategory = (Function->FunctionFlags & (FUNC_BlueprintCallable | FUNC_BlueprintPure)) != 0;
-		Tags.Emplace(FEmitHelper::HandleMetaData(Function, bMustHaveCategory, AdditinalMetaData));
-		Tags.Remove(FString());
+	static int32 ParseDelegateDetails(UFunction* Signature, FString& OutParametersMacro, FString& OutParamNumberStr);
 
-		FString AllTags;
-		FEmitHelper::ArrayToString(Tags, AllTags, TEXT(", "));
+	static TArray<FString> EmitSinglecastDelegateDeclarations(const TArray<UDelegateProperty*>& Delegates);
 
-		return FString::Printf(TEXT("UFUNCTION(%s)"), *AllTags);
-	}
+	static TArray<FString> EmitMulticastDelegateDeclarations(UClass* SourceClass);
 
-	static int32 ParseDelegateDetails(UFunction* Signature, FString& OutParametersMacro, FString& OutParamNumberStr)
-	{
-		int32 ParameterNum = 0;
-		FStringOutputDevice Parameters;
-		for (TFieldIterator<UProperty> PropIt(Signature); PropIt && (PropIt->PropertyFlags & CPF_Parm); ++PropIt)
-		{
-			Parameters += ", ";
-			PropIt->ExportCppDeclaration(Parameters, EExportedDeclaration::MacroParameter, NULL, EPropertyExportCPPFlags::CPPF_CustomTypeName);
-			++ParameterNum;
-		}
+	static FString EmitLifetimeReplicatedPropsImpl(UClass* SourceClass, const FString& CppClassName, const TCHAR* InCurrentIndent);
 
-		FString ParamNumberStr;
-		switch (ParameterNum)
-		{
-		case 0: ParamNumberStr = TEXT("");				break;
-		case 1: ParamNumberStr = TEXT("_OneParam");		break;
-		case 2: ParamNumberStr = TEXT("_TwoParams");	break;
-		case 3: ParamNumberStr = TEXT("_ThreeParams");	break;
-		case 4: ParamNumberStr = TEXT("_FourParams");	break;
-		case 5: ParamNumberStr = TEXT("_FiveParams");	break;
-		case 6: ParamNumberStr = TEXT("_SixParams");	break;
-		case 7: ParamNumberStr = TEXT("_SevenParams");	break;
-		case 8: ParamNumberStr = TEXT("_EightParams");	break;
-		default: ParamNumberStr = TEXT("_TooMany");		break;
-		}
+	static FString GatherNativeHeadersToInclude(UField* SourceItem, const TArray<FString>& PersistentHeaders);
 
-		OutParametersMacro = Parameters;
-		OutParamNumberStr = ParamNumberStr;
-		return ParameterNum;
-	}
+	static FString LiteralTerm(const FEdGraphPinType& Type, const FString& CustomValue, UObject* LiteralObject);
 
-	static TArray<FString> EmitSinglecastDelegateDeclarations(const TArray<UDelegateProperty*>& Delegates)
-	{
-		TArray<FString> Results;
-		for (auto It : Delegates)
-		{
-			check(It);
-			auto Signature = It->SignatureFunction;
-			check(Signature);
+	static FString DefaultValue(const FEdGraphPinType& Type);
 
-			FString ParamNumberStr, Parameters;
-			ParseDelegateDetails(Signature, Parameters, ParamNumberStr);
+	static UFunction* GetOriginalFunction(UFunction* Function);
 
-			Results.Add(*FString::Printf(TEXT("DECLARE_DYNAMIC_DELEGATE%s(%s%s)"),
-				*ParamNumberStr, *It->GetCPPType(NULL, EPropertyExportCPPFlags::CPPF_CustomTypeName), *Parameters));
-		}
-		return Results;
-	}
+	static bool ShoulsHandleAsNativeEvent(UFunction* Function);
 
-	static TArray<FString> EmitMulticastDelegateDeclarations(UClass* SourceClass)
-	{
-		TArray<FString> Results;
-		for (TFieldIterator<UMulticastDelegateProperty> It(SourceClass, EFieldIteratorFlags::ExcludeSuper); It; ++It)
-		{
-			auto Signature = It->SignatureFunction;
-			check(Signature);
-
-			FString ParamNumberStr, Parameters;
-			ParseDelegateDetails(Signature, Parameters, ParamNumberStr);
-
-			Results.Add(*FString::Printf(TEXT("DECLARE_DYNAMIC_MULTICAST_DELEGATE%s(%s%s)"),
-				*ParamNumberStr, *It->GetCPPType(NULL, EPropertyExportCPPFlags::CPPF_CustomTypeName), *Parameters));
-		}
-
-		return Results;
-	}
-
-	static FString EmitLifetimeReplicatedPropsImpl(UClass* SourceClass, const FString& CppClassName, const TCHAR* InCurrentIndent)
-	{
-		FString Result;
-		bool bFunctionInitilzed = false;
-		for (TFieldIterator<UProperty> It(SourceClass, EFieldIteratorFlags::ExcludeSuper); It; ++It)
-		{
-			if ((It->PropertyFlags & CPF_Net) != 0)
-			{
-				if (!bFunctionInitilzed)
-				{
-					Result += FString::Printf(TEXT("%svoid %s::GetLifetimeReplicatedProps(TArray< FLifetimeProperty > & OutLifetimeProps) const\n"), InCurrentIndent, *CppClassName);
-					Result += FString::Printf(TEXT("%s{\n"), InCurrentIndent);
-					Result += FString::Printf(TEXT("%s\tSuper::GetLifetimeReplicatedProps(OutLifetimeProps);\n"), InCurrentIndent);
-					bFunctionInitilzed = true;
-				}
-				Result += FString::Printf(TEXT("%s\tDOREPLIFETIME( %s, %s);\n"), InCurrentIndent, *CppClassName, *It->GetNameCPP());
-			}
-		}
-		if (bFunctionInitilzed)
-		{
-			Result += FString::Printf(TEXT("%s}\n"), InCurrentIndent);
-		}
-		return Result;
-	}
-
-	static FString GatherNativeHeadersToInclude(UField* SourceItem, const TArray<FString>& PersistentHeaders)
-	{
-		TSet<FString> HeaderFiles;
-		HeaderFiles.Append(PersistentHeaders);
-		{
-			TArray<UObject*> ReferencedObjects;
-			{
-				FReferenceFinder ReferenceFinder(ReferencedObjects, NULL, false, false, false, true);
-				TArray<UObject*> ObjectsToCheck;
-				GetObjectsWithOuter(SourceItem, ObjectsToCheck, true);
-				//CDO?
-				for (auto Obj : ObjectsToCheck)
-				{
-					ReferenceFinder.FindReferences(Obj);
-				}
-			}
-
-			auto EngineSourceDir = FPaths::EngineSourceDir();
-			auto GameSourceDir = FPaths::GameSourceDir();
-			auto CurrentPackage = SourceItem->GetTypedOuter<UPackage>();
-			for (auto Obj : ReferencedObjects)
-			{
-				auto Field = Cast<UField>(Obj);
-				FString PackPath;
-				if (Field
-					&& (Field->GetTypedOuter<UPackage>() != CurrentPackage)
-					&& FSourceCodeNavigation::FindClassHeaderPath(Field, PackPath))
-				{
-					if (!PackPath.RemoveFromStart(EngineSourceDir))
-					{
-						if (!PackPath.RemoveFromStart(GameSourceDir))
-						{
-							PackPath = FPaths::GetCleanFilename(PackPath);
-						}
-					}
-					HeaderFiles.Add(PackPath);
-				}
-			}
-		}
-
-		FString Result;
-		for (auto HeaderFile : HeaderFiles)
-		{
-			Result += FString::Printf(TEXT("#include \"%s\"\n"), *HeaderFile);
-		}
-
-		return Result;
-	}
-
-	static FString LiteralTerm(const FEdGraphPinType& Type, const FString& CustomValue, UObject* LiteralObject)
-	{
-		auto Schema = GetDefault<UEdGraphSchema_K2>();
-
-		if (UEdGraphSchema_K2::PC_String == Type.PinCategory)
-		{
-			return FString::Printf(TEXT("TEXT(\"%s\")"), *CustomValue);
-		}
-		else if (UEdGraphSchema_K2::PC_Text == Type.PinCategory)
-		{
-			return FString::Printf(TEXT("FText::FromString(TEXT(\"%s\"))"), *CustomValue);
-		}
-		else if (UEdGraphSchema_K2::PC_Float == Type.PinCategory)
-		{
-			float Value = CustomValue.IsEmpty() ? 0.0f : FCString::Atof(*CustomValue);
-			return FString::Printf(TEXT("%f"), Value);
-		}
-		else if (UEdGraphSchema_K2::PC_Int == Type.PinCategory)
-		{
-			int32 Value = CustomValue.IsEmpty() ? 0 : FCString::Atoi(*CustomValue);
-			return FString::Printf(TEXT("%d"), Value);
-		}
-		else if ((UEdGraphSchema_K2::PC_Byte == Type.PinCategory) || (UEdGraphSchema_K2::PC_Enum == Type.PinCategory))
-		{
-			auto TypeEnum = Cast<UEnum>(Type.PinSubCategoryObject.Get());
-			if (TypeEnum)
-			{
-				return FString::Printf(TEXT("%s::%s"), *TypeEnum->GetName(), CustomValue.IsEmpty() 
-					? *TypeEnum->GetEnumName(TypeEnum->NumEnums()-1)
-					: *CustomValue);
-			}
-			else
-			{
-				uint8 Value = CustomValue.IsEmpty() ? 0 : FCString::Atoi(*CustomValue);
-				return FString::Printf(TEXT("%u"), Value);
-			}
-		}
-		else if (UEdGraphSchema_K2::PC_Boolean == Type.PinCategory)
-		{
-			const bool bValue = CustomValue.ToBool();
-			return bValue ? TEXT("true") : TEXT("false");
-		}
-		else if (UEdGraphSchema_K2::PC_Name == Type.PinCategory)
-		{
-			return CustomValue.IsEmpty() 
-				? TEXT("FName()") 
-				: FString::Printf(TEXT("FName(TEXT(\"%s\"))"), *(FName(*CustomValue).ToString()));
-		}
-		else if (UEdGraphSchema_K2::PC_Struct == Type.PinCategory)
-		{
-			auto StructType = Cast<UScriptStruct>(Type.PinSubCategoryObject.Get());
-			ensure(StructType);
-			if (StructType == TBaseStructure<FVector>::Get())
-			{
-				FVector Vect = FVector::ZeroVector;
-				FDefaultValueHelper::ParseVector(CustomValue, /*out*/ Vect);
-				return FString::Printf(TEXT("FVector(%f,%f,%f)"), Vect.X, Vect.Y, Vect.Z);
-			}
-			else if (StructType == TBaseStructure<FRotator>::Get())
-			{
-				FRotator Rot = FRotator::ZeroRotator;
-				FDefaultValueHelper::ParseRotator(CustomValue, /*out*/ Rot);
-				return FString::Printf(TEXT("FRotator(%f,%f,%f)"), Rot.Pitch, Rot.Yaw, Rot.Roll);
-			}
-			else if (StructType == TBaseStructure<FTransform>::Get())
-			{
-				FTransform Trans = FTransform::Identity;
-				Trans.InitFromString(CustomValue);
-				const FQuat Rot = Trans.GetRotation();
-				const FVector Translation = Trans.GetTranslation();
-				const FVector Scale = Trans.GetScale3D();
-				return FString::Printf(TEXT("FTransform( FQuat(%f,%f,%f,%f), FVector(%f,%f,%f), FVector(%f,%f,%f) )"),
-					Rot.X, Rot.Y, Rot.Z, Rot.W, Translation.X, Translation.Y, Translation.Z, Scale.X, Scale.Y, Scale.Z);
-			}
-			else if (StructType == TBaseStructure<FLinearColor>::Get())
-			{
-				FLinearColor LinearColor;
-				LinearColor.InitFromString(CustomValue);
-				return FString::Printf(TEXT("FLinearColor(%f,%f,%f,%f)"), LinearColor.R, LinearColor.G, LinearColor.B, LinearColor.A);
-			}
-			else if (StructType == TBaseStructure<FColor>::Get())
-			{
-				FColor Color;
-				Color.InitFromString(CustomValue);
-				return FString::Printf(TEXT("FColor(%d,%d,%d,%d)"), Color.R, Color.G, Color.B, Color.A);
-			}
-			if (StructType == TBaseStructure<FVector2D>::Get())
-			{
-				FVector2D Vect = FVector2D::ZeroVector;
-				Vect.InitFromString(CustomValue);
-				return FString::Printf(TEXT("FVector2D(%f,%f)"), Vect.X, Vect.Y);
-			}
-			else
-			{
-				//@todo:  This needs to be more robust, since import text isn't really proper for struct construction.
-				return FString(TEXT("F")) + StructType->GetName() + (CustomValue.IsEmpty() ? TEXT("{}") : *CustomValue);
-			}
-		}
-		else if (Type.PinSubCategory == UEdGraphSchema_K2::PSC_Self)
-		{
-			return TEXT("this");
-		}
-		else if (UEdGraphSchema_K2::PC_Class == Type.PinCategory)
-		{
-			if (auto FoundClass = Cast<const UClass>(LiteralObject))
-			{
-				return FString::Printf(TEXT("%s%s::StaticClass()"), FoundClass->GetPrefixCPP(), *FoundClass->GetName());
-			}
-			return FString(TEXT("nullptr"));
-		}
-		else if((UEdGraphSchema_K2::PC_AssetClass == Type.PinCategory) || (UEdGraphSchema_K2::PC_Asset == Type.PinCategory))
-		{
-			return LiteralObject
-				? FString::Printf(TEXT("FStringAssetReference(TEXT(\"%s\"))"), *(LiteralObject->GetPathName()))
-				: FString(TEXT("nullptr"));
-		}
-		else if (UEdGraphSchema_K2::PC_Object == Type.PinCategory)
-		{
-			if (LiteralObject)
-			{
-				UClass* FoundClass = Cast<UClass>(Type.PinSubCategoryObject.Get());
-				FString ClassString = FoundClass ? (FString(FoundClass->GetPrefixCPP()) + FoundClass->GetName()) : TEXT("UObject");
-				return FString::Printf(TEXT("FindObject<%s>(ANY_PACKAGE, TEXT(\"%s\"))"), *ClassString, *(LiteralObject->GetPathName()));
-			}
-			else
-			{
-				return FString(TEXT("nullptr"));
-			}
-		}
-		/*
-		else if (CoerceProperty->IsA(UInterfaceProperty::StaticClass()))
-		{
-			if (Term->Type.PinSubCategory == PSC_Self)
-			{
-				return TEXT("this");
-			}
-			else
-			{
-				ensureMsgf(false, TEXT("It is not possible to express this interface property as a literal value!"));
-				return Term->Name;
-			}
-		}
-		*/
-		ensureMsgf(false, TEXT("It is not possible to express this type as a literal value!"));
-		return CustomValue;
-	}
-
-	static FString DefaultValue(const FEdGraphPinType& Type)
-	{
-		return LiteralTerm(Type, FString(), nullptr);
-	}
-
-	static UFunction* GetOriginalFunction(UFunction* Function)
-	{
-		check(Function);
-		const FName FunctionName = Function->GetFName();
-
-		UClass* Owner = Function->GetOwnerClass();
-		check(Owner);
-		for (auto& Inter : Owner->Interfaces)
-		{
-			if (UFunction* Result = Inter.Class->FindFunctionByName(FunctionName))
-			{
-				return GetOriginalFunction(Result);
-			}
-		}
-
-		for (UClass* SearchClass = Owner->GetSuperClass(); SearchClass; SearchClass = SearchClass->GetSuperClass())
-		{
-			if (UFunction* Result = SearchClass->FindFunctionByName(FunctionName))
-			{
-				return GetOriginalFunction(Result);
-			}
-		}
-
-		return Function;
-	}
-
-	static bool ShoulsHandleAsNativeEvent(UFunction* Function)
-	{
-		check(Function);
-		auto OriginalFunction = FEmitHelper::GetOriginalFunction(Function);
-		check(OriginalFunction);
-		if (OriginalFunction != Function)
-		{
-			const uint32 FlagsToCheckMask = EFunctionFlags::FUNC_Event | EFunctionFlags::FUNC_BlueprintEvent | EFunctionFlags::FUNC_Native;
-			const uint32 FlagsToCheck = OriginalFunction->FunctionFlags & FlagsToCheckMask;
-
-			auto OriginalOwner = OriginalFunction->GetOwnerClass();
-			const bool bBPInterface = Cast<UBlueprintGeneratedClass>(OriginalOwner) && OriginalOwner->IsChildOf<UInterface>();
-			return (FlagsToCheck == FlagsToCheckMask) || bBPInterface;
-		}
-		return false;
-	}
+	static bool ShoulsHandleAsImplementableEvent(UFunction* Function);
+};
 
-	static bool ShoulsHandleAsImplementableEvent(UFunction* Function)
-	{
-		check(Function);
-		auto OriginalFunction = FEmitHelper::GetOriginalFunction(Function);
-		check(OriginalFunction);
-		if (OriginalFunction != Function)
-		{
-			const uint32 FlagsToCheckMask = EFunctionFlags::FUNC_Event | EFunctionFlags::FUNC_BlueprintEvent | EFunctionFlags::FUNC_Native;
-			const uint32 FlagsToCheck = OriginalFunction->FunctionFlags & FlagsToCheckMask;
-			return (FlagsToCheck == (EFunctionFlags::FUNC_Event | EFunctionFlags::FUNC_BlueprintEvent));
-		}
-		return false;
-	}
+struct FEmitDefaultValueHelper
+{
+	static FString GenerateGetDefaultValue(const UUserDefinedStruct* Struct);
+	static void InnerGenerate(const UProperty* Property, const uint8* ValuePtr, const FString& PathToMember, FString& OutResult);
 };
\ No newline at end of file
diff --git a/Engine/Source/Runtime/Core/Private/Misc/DateTime.cpp b/Engine/Source/Runtime/Core/Private/Misc/DateTime.cpp
--- a/Engine/Source/Runtime/Core/Private/Misc/DateTime.cpp
+++ b/Engine/Source/Runtime/Core/Private/Misc/DateTime.cpp
@@ -1,7 +1,7 @@
 ﻿// Copyright 1998-2015 Epic Games, Inc. All Rights Reserved.
 
 #include "CorePrivatePCH.h"
-
+#include "PropertyPortFlags.h"
 
 /* FDateTime constants
  *****************************************************************************/
@@ -47,6 +47,12 @@ FDateTime::FDateTime( int32 Year, int32 Month, int32 Day, int32 Hour, int32 Minu
 
 bool FDateTime::ExportTextItem( FString& ValueStr, FDateTime const& DefaultValue, UObject* Parent, int32 PortFlags, UObject* ExportRootScope ) const
 {
+	if (0 != (PortFlags & EPropertyPortFlags::PPF_ExportCpp))
+	{
+		ValueStr += FString::Printf(TEXT("FDateTime(0x%016X)"), Ticks);
+		return true;
+	}
+
 	ValueStr += ToString();
 
 	return true;
diff --git a/Engine/Source/Runtime/Core/Private/Misc/Guid.cpp b/Engine/Source/Runtime/Core/Private/Misc/Guid.cpp
--- a/Engine/Source/Runtime/Core/Private/Misc/Guid.cpp
+++ b/Engine/Source/Runtime/Core/Private/Misc/Guid.cpp
@@ -1,13 +1,18 @@
 // Copyright 1998-2015 Epic Games, Inc. All Rights Reserved.
 
 #include "CorePrivatePCH.h"
-
+#include "PropertyPortFlags.h"
 
 /* FGuid interface
  *****************************************************************************/
 
 bool FGuid::ExportTextItem( FString& ValueStr, FGuid const& DefaultValue, UObject* Parent, int32 PortFlags, class UObject* ExportRootScope ) const
 {
+	if (0 != (PortFlags & EPropertyPortFlags::PPF_ExportCpp))
+	{
+		return false;
+	}
+
 	ValueStr += ToString();
 
 	return true;
diff --git a/Engine/Source/Runtime/Core/Private/Misc/Timespan.cpp b/Engine/Source/Runtime/Core/Private/Misc/Timespan.cpp
--- a/Engine/Source/Runtime/Core/Private/Misc/Timespan.cpp
+++ b/Engine/Source/Runtime/Core/Private/Misc/Timespan.cpp
@@ -1,13 +1,19 @@
 // Copyright 1998-2015 Epic Games, Inc. All Rights Reserved.
 
 #include "CorePrivatePCH.h"
-
+#include "PropertyPortFlags.h"
 
 /* FTimespan interface
  *****************************************************************************/
 
 bool FTimespan::ExportTextItem( FString& ValueStr, FTimespan const& DefaultValue, UObject* Parent, int32 PortFlags, UObject* ExportRootScope ) const
 {
+	if (0 != (PortFlags & EPropertyPortFlags::PPF_ExportCpp))
+	{
+		ValueStr += FString::Printf(TEXT("FTimespan(0x%016X)"), Ticks);
+		return true;
+	}
+
 	ValueStr += ToString(TEXT("%N%d.%h:%m:%s.%f"));
 
 	return true;
diff --git a/Engine/Source/Runtime/Core/Public/UObject/PropertyPortFlags.h b/Engine/Source/Runtime/Core/Public/UObject/PropertyPortFlags.h
--- a/Engine/Source/Runtime/Core/Public/UObject/PropertyPortFlags.h
+++ b/Engine/Source/Runtime/Core/Public/UObject/PropertyPortFlags.h
@@ -93,4 +93,7 @@ enum EPropertyPortFlags
 
 	/** Ignores CPF_Deprecated flag */
 	PPF_UseDeprecatedProperties		= 0x08000000,
+
+	/** Export in C++ form */
+	PPF_ExportCpp					= 0x10000000,
 };
\ No newline at end of file
diff --git a/Engine/Source/Runtime/CoreUObject/Private/Misc/StringAssetReference.cpp b/Engine/Source/Runtime/CoreUObject/Private/Misc/StringAssetReference.cpp
--- a/Engine/Source/Runtime/CoreUObject/Private/Misc/StringAssetReference.cpp
+++ b/Engine/Source/Runtime/CoreUObject/Private/Misc/StringAssetReference.cpp
@@ -71,6 +71,11 @@ void FStringAssetReference::operator=(FStringAssetReference const& Other)
 }
 bool FStringAssetReference::ExportTextItem(FString& ValueStr, FStringAssetReference const& DefaultValue, UObject* Parent, int32 PortFlags, UObject* ExportRootScope) const
 {
+	if (0 != (PortFlags & EPropertyPortFlags::PPF_ExportCpp))
+	{
+		return false;
+	}
+
 	if (IsValid())
 	{
 		ValueStr += ToString();
diff --git a/Engine/Source/Runtime/CoreUObject/Private/UObject/Property.cpp b/Engine/Source/Runtime/CoreUObject/Private/UObject/Property.cpp
--- a/Engine/Source/Runtime/CoreUObject/Private/UObject/Property.cpp
+++ b/Engine/Source/Runtime/CoreUObject/Private/UObject/Property.cpp
@@ -1261,6 +1261,16 @@ void UNumericProperty::ExportTextItem( FString& ValueStr, const void* PropertyVa
 	ValueStr += GetNumericPropertyValueToString(PropertyValue);
 }
 
+void UFloatProperty::ExportTextItem(FString& ValueStr, const void* PropertyValue, const void* DefaultValue, UObject* Parent, int32 PortFlags, UObject* ExportRootScope) const
+{
+	Super::ExportTextItem(ValueStr, PropertyValue, DefaultValue, Parent, PortFlags, ExportRootScope);
+
+	if (0 != (PortFlags & PPF_ExportCpp))
+	{
+		ValueStr += TEXT("f");
+	}
+}
+
 
 IMPLEMENT_CORE_INTRINSIC_CLASS(UNumericProperty, UProperty,
 {
diff --git a/Engine/Source/Runtime/CoreUObject/Private/UObject/PropertyArray.cpp b/Engine/Source/Runtime/CoreUObject/Private/UObject/PropertyArray.cpp
--- a/Engine/Source/Runtime/CoreUObject/Private/UObject/PropertyArray.cpp
+++ b/Engine/Source/Runtime/CoreUObject/Private/UObject/PropertyArray.cpp
@@ -125,6 +125,14 @@ void UArrayProperty::ExportTextItem( FString& ValueStr, const void* PropertyValu
 {
 	checkSlow(Inner);
 
+	if (0 != (PortFlags & PPF_ExportCpp))
+	{
+		FString ExtendedTypeText;
+		FString TypeText = GetCPPType(&ExtendedTypeText, EPropertyExportCPPFlags::CPPF_BlueprintCppBackend);
+		ValueStr += FString::Printf(TEXT("%s%s()"), *TypeText, *ExtendedTypeText);
+		return;
+	}
+
 	FScriptArrayHelper ArrayHelper(this, PropertyValue);
 	FScriptArrayHelper DefaultArrayHelper(this, DefaultValue);
 
diff --git a/Engine/Source/Runtime/CoreUObject/Private/UObject/PropertyAssetObject.cpp b/Engine/Source/Runtime/CoreUObject/Private/UObject/PropertyAssetObject.cpp
--- a/Engine/Source/Runtime/CoreUObject/Private/UObject/PropertyAssetObject.cpp
+++ b/Engine/Source/Runtime/CoreUObject/Private/UObject/PropertyAssetObject.cpp
@@ -72,6 +72,12 @@ void UAssetObjectProperty::ExportTextItem( FString& ValueStr, const void* Proper
 		ID = AssetPtr.GetUniqueID();
 	}
 
+	if (0 != (PortFlags & PPF_ExportCpp))
+	{
+		ValueStr += FString::Printf(TEXT("FStringAssetReference(TEXT(\"%s\"))"), *ID.ToString().ReplaceCharWithEscapedChar());
+		return;
+	}
+
 	if (!ID.ToString().IsEmpty())
 	{
 		ValueStr += ID.ToString();
diff --git a/Engine/Source/Runtime/CoreUObject/Private/UObject/PropertyBaseObject.cpp b/Engine/Source/Runtime/CoreUObject/Private/UObject/PropertyBaseObject.cpp
--- a/Engine/Source/Runtime/CoreUObject/Private/UObject/PropertyBaseObject.cpp
+++ b/Engine/Source/Runtime/CoreUObject/Private/UObject/PropertyBaseObject.cpp
@@ -172,6 +172,20 @@ FString UObjectPropertyBase::GetExportPath(const UObject* Object, const UObject*
 void UObjectPropertyBase::ExportTextItem( FString& ValueStr, const void* PropertyValue, const void* DefaultValue, UObject* Parent, int32 PortFlags, UObject* ExportRootScope ) const
 {
 	UObject* Temp = GetObjectPropertyValue(PropertyValue);
+
+	if (0 != (PortFlags & PPF_ExportCpp))
+	{
+		FString::Printf(TEXT("%s%s*"), PropertyClass->GetPrefixCPP(), *PropertyClass->GetName());
+
+		ValueStr += Temp
+			? FString::Printf(TEXT("LoadObject<%s%s>(nullptr, TEXT(\"%s\"))")
+				, PropertyClass->GetPrefixCPP()
+				, *PropertyClass->GetName()
+				, *(Temp->GetPathName().ReplaceCharWithEscapedChar()))
+			: TEXT("nullptr");
+		return;
+	}
+
 	if( Temp != NULL )
 	{
 		if (PortFlags & PPF_DebugDump)
diff --git a/Engine/Source/Runtime/CoreUObject/Private/UObject/PropertyBool.cpp b/Engine/Source/Runtime/CoreUObject/Private/UObject/PropertyBool.cpp
--- a/Engine/Source/Runtime/CoreUObject/Private/UObject/PropertyBool.cpp
+++ b/Engine/Source/Runtime/CoreUObject/Private/UObject/PropertyBool.cpp
@@ -188,7 +188,16 @@ void UBoolProperty::ExportTextItem( FString& ValueStr, const void* PropertyValue
 {
 	check(FieldSize != 0);
 	const uint8* ByteValue = (uint8*)PropertyValue + ByteOffset;
-	TCHAR* Temp =	(TCHAR*) (((*ByteValue) & FieldMask) ? TEXT("True") : TEXT("False"));
+	const bool bValue = 0 != ((*ByteValue) & FieldMask);
+	TCHAR* Temp = nullptr;
+	if (0 != (PortFlags & PPF_ExportCpp))
+	{
+		Temp = (bValue ? TEXT("true") : TEXT("false"));
+	}
+	else
+	{
+		Temp = (bValue ? TEXT("True") : TEXT("False"));
+	}
 	ValueStr += FString::Printf( TEXT("%s"), Temp );
 }
 const TCHAR* UBoolProperty::ImportText_Internal( const TCHAR* Buffer, void* Data, int32 PortFlags, UObject* Parent, FOutputDevice* ErrorText ) const
diff --git a/Engine/Source/Runtime/CoreUObject/Private/UObject/PropertyByte.cpp b/Engine/Source/Runtime/CoreUObject/Private/UObject/PropertyByte.cpp
--- a/Engine/Source/Runtime/CoreUObject/Private/UObject/PropertyByte.cpp
+++ b/Engine/Source/Runtime/CoreUObject/Private/UObject/PropertyByte.cpp
@@ -109,6 +109,23 @@ FString UByteProperty::GetCPPType( FString* ExtendedTypeText/*=NULL*/, uint32 CP
 
 void UByteProperty::ExportTextItem( FString& ValueStr, const void* PropertyValue, const void* DefaultValue, UObject* Parent, int32 PortFlags, UObject* ExportRootScope ) const
 {
+	if (0 != (PortFlags & PPF_ExportCpp))
+	{
+		if (Enum)
+		{
+			const int32 ActualValue = *(const uint8*)PropertyValue;
+			const int32 MaxValue = Enum->GetMaxEnumValue();
+			const int32 GoodValue = (ActualValue > MaxValue) ? MaxValue : ActualValue;
+			ValueStr += FString::Printf(TEXT("%s::%s"), *Enum->GetName(),
+				*Enum->GetEnumName(Enum->GetIndexByValue(GoodValue)));
+		}
+		else
+		{
+			Super::ExportTextItem(ValueStr, PropertyValue, DefaultValue, Parent, PortFlags, ExportRootScope);
+		}
+		return;
+	}
+
 	if( Enum && (PortFlags & PPF_ConsoleVariable) == 0 )
 	{
 		// if the value is the max value (the autogenerated *_MAX value), export as "INVALID", unless we're exporting text for copy/paste (for copy/paste,
diff --git a/Engine/Source/Runtime/CoreUObject/Private/UObject/PropertyDelegate.cpp b/Engine/Source/Runtime/CoreUObject/Private/UObject/PropertyDelegate.cpp
--- a/Engine/Source/Runtime/CoreUObject/Private/UObject/PropertyDelegate.cpp
+++ b/Engine/Source/Runtime/CoreUObject/Private/UObject/PropertyDelegate.cpp
@@ -89,6 +89,12 @@ FString UDelegateProperty::GetCPPTypeForwardDeclaration() const
 
 void UDelegateProperty::ExportTextItem( FString& ValueStr, const void* PropertyValue, const void* DefaultValue, UObject* Parent, int32 PortFlags, UObject* ExportRootScope ) const
 {
+	if (0 != (PortFlags & PPF_ExportCpp))
+	{
+		ValueStr += TEXT("{}");
+		return;
+	}
+
 	FScriptDelegate* ScriptDelegate = (FScriptDelegate*)PropertyValue;
 	check(ScriptDelegate != NULL);
 	bool bDelegateHasValue = ScriptDelegate->GetFunctionName() != NAME_None;
diff --git a/Engine/Source/Runtime/CoreUObject/Private/UObject/PropertyInterface.cpp b/Engine/Source/Runtime/CoreUObject/Private/UObject/PropertyInterface.cpp
--- a/Engine/Source/Runtime/CoreUObject/Private/UObject/PropertyInterface.cpp
+++ b/Engine/Source/Runtime/CoreUObject/Private/UObject/PropertyInterface.cpp
@@ -62,7 +62,7 @@ FString UInterfaceProperty::GetCPPType( FString* ExtendedTypeText/*=NULL*/, uint
 			}
 		}
 		check(ExportClass);
-		check(ExportClass->HasAnyClassFlags(CLASS_Interface));
+		check(ExportClass->HasAnyClassFlags(CLASS_Interface) || 0 != (CPPF_BlueprintCppBackend & CPPExportFlags));
 
 		*ExtendedTypeText = FString::Printf(TEXT("<I%s>"), *ExportClass->GetName());
 	}
@@ -133,6 +133,18 @@ void UInterfaceProperty::ExportTextItem( FString& ValueStr, const void* Property
 	FScriptInterface* InterfaceValue = (FScriptInterface*)PropertyValue;
 
 	UObject* Temp = InterfaceValue->GetObject();
+
+	if (0 != (PortFlags & PPF_ExportCpp))
+	{
+		const FString GetObjectStr = Temp
+			? FString::Printf(TEXT("LoadObject<UObject>(nullptr, TEXT(\"%s\"))"), *Temp->GetPathName().ReplaceCharWithEscapedChar())
+			: TEXT("");
+		ValueStr += FString::Printf(TEXT("TScriptInterface<I%s>(%s)")
+			, (InterfaceClass ? *InterfaceClass->GetName() : TEXT("Interface"))
+			, *GetObjectStr);
+		return;
+	}
+
 	if( Temp != NULL )
 	{
 		bool bExportFullyQualified = true;
diff --git a/Engine/Source/Runtime/CoreUObject/Private/UObject/PropertyMap.cpp b/Engine/Source/Runtime/CoreUObject/Private/UObject/PropertyMap.cpp
--- a/Engine/Source/Runtime/CoreUObject/Private/UObject/PropertyMap.cpp
+++ b/Engine/Source/Runtime/CoreUObject/Private/UObject/PropertyMap.cpp
@@ -452,6 +452,12 @@ FString UMapProperty::GetCPPMacroType( FString& ExtendedTypeText ) const
 
 void UMapProperty::ExportTextItem(FString& ValueStr, const void* PropertyValue, const void* DefaultValue, UObject* Parent, int32 PortFlags, UObject* ExportRootScope) const
 {
+	if (0 != (PortFlags & PPF_ExportCpp))
+	{
+		ValueStr += TEXT("{}");
+		return;
+	}
+
 	checkSlow(KeyProp);
 	checkSlow(ValueProp);
 
diff --git a/Engine/Source/Runtime/CoreUObject/Private/UObject/PropertyMulticastDelegate.cpp b/Engine/Source/Runtime/CoreUObject/Private/UObject/PropertyMulticastDelegate.cpp
--- a/Engine/Source/Runtime/CoreUObject/Private/UObject/PropertyMulticastDelegate.cpp
+++ b/Engine/Source/Runtime/CoreUObject/Private/UObject/PropertyMulticastDelegate.cpp
@@ -133,6 +133,12 @@ FString UMulticastDelegateProperty::GetCPPType( FString* ExtendedTypeText/*=NULL
 
 void UMulticastDelegateProperty::ExportTextItem( FString& ValueStr, const void* PropertyValue, const void* DefaultValue, UObject* Parent, int32 PortFlags, UObject* ExportRootScope ) const
 {
+	if (0 != (PortFlags & PPF_ExportCpp))
+	{
+		ValueStr += TEXT("{}");
+		return;
+	}
+
 	const FMulticastScriptDelegate* MulticastDelegate = (const FMulticastScriptDelegate*)( PropertyValue );
 	check( MulticastDelegate != NULL );
 
diff --git a/Engine/Source/Runtime/CoreUObject/Private/UObject/PropertyName.cpp b/Engine/Source/Runtime/CoreUObject/Private/UObject/PropertyName.cpp
--- a/Engine/Source/Runtime/CoreUObject/Private/UObject/PropertyName.cpp
+++ b/Engine/Source/Runtime/CoreUObject/Private/UObject/PropertyName.cpp
@@ -9,7 +9,13 @@
 void UNameProperty::ExportTextItem( FString& ValueStr, const void* PropertyValue, const void* DefaultValue, UObject* Parent, int32 PortFlags, UObject* ExportRootScope ) const
 {
 	FName Temp = *(FName*)PropertyValue;
-	if( !(PortFlags & PPF_Delimited) )
+	if (0 != (PortFlags & PPF_ExportCpp))
+	{
+		ValueStr += (Temp == NAME_None) 
+			? TEXT("FName()") 
+			: FString::Printf(TEXT("FName(TEXT(\"%s\"))"), *(Temp.ToString().ReplaceCharWithEscapedChar()));
+	}
+	else if( !(PortFlags & PPF_Delimited) )
 	{
 		ValueStr += Temp.ToString();
 	}
diff --git a/Engine/Source/Runtime/CoreUObject/Private/UObject/PropertyStr.cpp b/Engine/Source/Runtime/CoreUObject/Private/UObject/PropertyStr.cpp
--- a/Engine/Source/Runtime/CoreUObject/Private/UObject/PropertyStr.cpp
+++ b/Engine/Source/Runtime/CoreUObject/Private/UObject/PropertyStr.cpp
@@ -9,7 +9,11 @@
 void UStrProperty::ExportTextItem( FString& ValueStr, const void* PropertyValue, const void* DefaultValue, UObject* Parent, int32 PortFlags, UObject* ExportRootScope ) const
 {
 	FString& StringValue = *(FString*)PropertyValue;
-	if( !(PortFlags & PPF_Delimited) )
+	if (0 != (PortFlags & PPF_ExportCpp))
+	{
+		ValueStr += FString::Printf(TEXT("FString(TEXT(\"%s\"))"), *(StringValue.ReplaceCharWithEscapedChar()));
+	}
+	else if (!(PortFlags & PPF_Delimited))
 	{
 		ValueStr += StringValue;
 	}
diff --git a/Engine/Source/Runtime/CoreUObject/Private/UObject/PropertyStruct.cpp b/Engine/Source/Runtime/CoreUObject/Private/UObject/PropertyStruct.cpp
--- a/Engine/Source/Runtime/CoreUObject/Private/UObject/PropertyStruct.cpp
+++ b/Engine/Source/Runtime/CoreUObject/Private/UObject/PropertyStruct.cpp
@@ -328,6 +328,12 @@ void UStructProperty::ExportTextItem( FString& ValueStr, const void* PropertyVal
 		}
 	}
 
+	if (0 != (PortFlags & PPF_ExportCpp))
+	{
+		ValueStr += FString::Printf(TEXT("F%s()"), *Struct->GetName());
+		return;
+	}
+
 	UStructProperty_ExportTextItem(Struct, ValueStr, PropertyValue, DefaultValue, Parent, PortFlags, ExportRootScope);
 } 
 
diff --git a/Engine/Source/Runtime/CoreUObject/Private/UObject/UTextProperty.cpp b/Engine/Source/Runtime/CoreUObject/Private/UObject/UTextProperty.cpp
--- a/Engine/Source/Runtime/CoreUObject/Private/UObject/UTextProperty.cpp
+++ b/Engine/Source/Runtime/CoreUObject/Private/UObject/UTextProperty.cpp
@@ -49,7 +49,11 @@ void UTextProperty::ExportTextItem( FString& ValueStr, const void* PropertyValue
 {
 	const FText& TextValue = GetPropertyValue(PropertyValue);
 	const FString& StringValue = FTextInspector::GetDisplayString(TextValue);
-	if( !(PortFlags & PPF_Delimited) )
+	if (0 != (PortFlags & PPF_ExportCpp))
+	{
+		ValueStr += FString::Printf(TEXT("FText::FromString(TEXT(\"%s\"))"), *(StringValue.ReplaceCharWithEscapedChar()));
+	}
+	else if( !(PortFlags & PPF_Delimited) )
 	{
 		ValueStr += StringValue;
 	}
diff --git a/Engine/Source/Runtime/CoreUObject/Public/UObject/UnrealType.h b/Engine/Source/Runtime/CoreUObject/Public/UObject/UnrealType.h
--- a/Engine/Source/Runtime/CoreUObject/Public/UObject/UnrealType.h
+++ b/Engine/Source/Runtime/CoreUObject/Public/UObject/UnrealType.h
@@ -31,7 +31,7 @@ enum EPropertyExportCPPFlags
 	CPPF_NoConst					=	0x00000010,
 	/** No reference '&' sign */
 	CPPF_NoRef						=	0x00000020,
-	/** Blueprint kompiler generated C++ code */
+	/** Blueprint compiler generated C++ code */
 	CPPF_BlueprintCppBackend		=	0x00000080,
 };
 
@@ -1476,6 +1476,10 @@ class COREUOBJECT_API UFloatProperty : public TProperty_Numeric<float>
 		:	TProperty_Numeric( ObjectInitializer, EC_CppProperty, InOffset, InFlags )
 	{
 	}
+
+	// UProperty interface
+	virtual void ExportTextItem(FString& ValueStr, const void* PropertyValue, const void* DefaultValue, UObject* Parent, int32 PortFlags, UObject* ExportRootScope) const override;
+	// End of UProperty interface
 };
 
 /*-----------------------------------------------------------------------------
diff --git a/Engine/Source/Runtime/Engine/Private/OnlineReplStructs.cpp b/Engine/Source/Runtime/Engine/Private/OnlineReplStructs.cpp
--- a/Engine/Source/Runtime/Engine/Private/OnlineReplStructs.cpp
+++ b/Engine/Source/Runtime/Engine/Private/OnlineReplStructs.cpp
@@ -53,6 +53,11 @@ bool FUniqueNetIdRepl::Serialize(FArchive& Ar)
 
 bool FUniqueNetIdRepl::ExportTextItem(FString& ValueStr, FUniqueNetIdRepl const& DefaultValue, UObject* Parent, int32 PortFlags, UObject* ExportRootScope) const
 {
+	if (0 != (PortFlags & EPropertyPortFlags::PPF_ExportCpp))
+	{
+		return false;
+	}
+
 	ValueStr += UniqueNetId.IsValid() ? UniqueNetId->ToString() : TEXT("INVALID");
 	return true;
 }
diff --git a/Engine/Source/Runtime/IOS/IOSRuntimeSettings/Classes/IOSRuntimeSettings.h b/Engine/Source/Runtime/IOS/IOSRuntimeSettings/Classes/IOSRuntimeSettings.h
--- a/Engine/Source/Runtime/IOS/IOSRuntimeSettings/Classes/IOSRuntimeSettings.h
+++ b/Engine/Source/Runtime/IOS/IOSRuntimeSettings/Classes/IOSRuntimeSettings.h
@@ -48,6 +48,11 @@ struct FIOSBuildResourceFilePath
 	 */
 	bool ExportTextItem(FString& ValueStr, FIOSBuildResourceFilePath const& DefaultValue, UObject* Parent, int32 PortFlags, UObject* ExportRootScope) const
 	{
+		if (0 != (PortFlags & EPropertyPortFlags::PPF_ExportCpp))
+		{
+			return false;
+		}
+
 		ValueStr += FilePath;
 		return true;
 	}
@@ -96,6 +101,11 @@ struct FIOSBuildResourceDirectory
 	 */
 	bool ExportTextItem(FString& ValueStr, FIOSBuildResourceDirectory const& DefaultValue, UObject* Parent, int32 PortFlags, UObject* ExportRootScope) const
 	{
+		if (0 != (PortFlags & EPropertyPortFlags::PPF_ExportCpp))
+		{
+			return false;
+		}
+
 		ValueStr += Path;
 		return true;
 	}
diff --git a/Engine/Source/Runtime/InputCore/Private/InputCoreTypes.cpp b/Engine/Source/Runtime/InputCore/Private/InputCoreTypes.cpp
--- a/Engine/Source/Runtime/InputCore/Private/InputCoreTypes.cpp
+++ b/Engine/Source/Runtime/InputCore/Private/InputCoreTypes.cpp
@@ -1018,6 +1018,11 @@ bool FKey::SerializeFromMismatchedTag(struct FPropertyTag const& Tag, FArchive&
 
 bool FKey::ExportTextItem(FString& ValueStr, FKey const& DefaultValue, UObject* Parent, int32 PortFlags, UObject* ExportRootScope) const
 {
+	if (0 != (PortFlags & EPropertyPortFlags::PPF_ExportCpp))
+	{
+		return false;
+	}
+
 	ValueStr += KeyName.ToString();
 	return true;
 }
