added the missing BlueprintCompilerCppBackendUtils.cpp file

diff --git a/Engine/Source/Developer/BlueprintCompilerCppBackend/Private/BlueprintCompilerCppBackendUtils.cpp b/Engine/Source/Developer/BlueprintCompilerCppBackend/Private/BlueprintCompilerCppBackendUtils.cpp
new file mode 100644
--- /dev/null
+++ b/Engine/Source/Developer/BlueprintCompilerCppBackend/Private/BlueprintCompilerCppBackendUtils.cpp
@@ -0,0 +1,711 @@
+// Copyright 1998-2015 Epic Games, Inc. All Rights Reserved.
+
+#include "BlueprintCompilerCppBackendModulePrivatePCH.h"
+#include "BlueprintCompilerCppBackendUtils.h"
+#include "EdGraphSchema_K2.h"
+#include "Editor/UnrealEd/Public/Kismet2/StructureEditorUtils.h"
+
+FString FSafeContextScopedEmmitter::GetAdditionalIndent() const
+{
+	return bSafeContextUsed ? TEXT("\t") : FString();
+}
+
+bool FSafeContextScopedEmmitter::IsSafeContextUsed() const
+{
+	return bSafeContextUsed;
+}
+
+FString FSafeContextScopedEmmitter::ValidationChain(const FBPTerminal* Term, FBlueprintCompilerCppBackend& CppBackend)
+{
+	TArray<FString> SafetyConditions;
+	for (; Term; Term = Term->Context)
+	{
+		if (!Term->IsStructContextType() && (Term->Type.PinSubCategory != TEXT("self")))
+		{
+			ensure(!Term->bIsLiteral);
+			SafetyConditions.Add(CppBackend.TermToText(Term, false));
+		}
+	}
+
+	FString Result;
+	for (int32 Iter = SafetyConditions.Num() - 1; Iter >= 0; --Iter)
+	{
+		Result += FString(TEXT("IsValid("));
+		Result += SafetyConditions[Iter];
+		Result += FString(TEXT(")"));
+		if (Iter)
+		{
+			Result += FString(TEXT(" && "));
+		}
+	}
+
+	return Result;
+}
+
+FSafeContextScopedEmmitter::FSafeContextScopedEmmitter(FString& InBody, const FBPTerminal* Term, FBlueprintCompilerCppBackend& CppBackend, const TCHAR* InCurrentIndent)
+	: Body(InBody), bSafeContextUsed(false), CurrentIndent(InCurrentIndent)
+{
+	const FString Conditions = ValidationChain(Term, CppBackend);
+
+	if (!Conditions.IsEmpty())
+	{
+		bSafeContextUsed = true;
+		Body += CurrentIndent;
+		Body += TEXT("if (");
+		Body += Conditions;
+		Body += TEXT(")\n");
+		Body += CurrentIndent;
+		Body += TEXT("{\n");
+	}
+}
+
+FSafeContextScopedEmmitter::~FSafeContextScopedEmmitter()
+{
+	if (bSafeContextUsed)
+	{
+		Body += CurrentIndent;
+		Body += TEXT("}\n");
+	}
+}
+
+void FEmitHelper::ArrayToString(const TArray<FString>& Array, FString& String, const TCHAR* Separator)
+{
+	if (Array.Num())
+	{
+		String += Array[0];
+	}
+	for (int32 i = 1; i < Array.Num(); ++i)
+	{
+		String += TEXT(", ");
+		String += Array[i];
+	}
+}
+
+bool FEmitHelper::HasAllFlags(uint64 Flags, uint64 FlagsToCheck)
+{
+	return FlagsToCheck == (Flags & FlagsToCheck);
+}
+
+FString FEmitHelper::HandleRepNotifyFunc(const UProperty* Property)
+{
+	check(Property);
+	if (HasAllFlags(Property->PropertyFlags, CPF_Net | CPF_RepNotify))
+	{
+		return FString::Printf(TEXT("ReplicatedUsing=%s"), *Property->RepNotifyFunc.ToString());
+	}
+	else if (HasAllFlags(Property->PropertyFlags, CPF_Net))
+	{
+		return TEXT("Replicated");
+	}
+	return FString();
+}
+
+bool FEmitHelper::MetaDataCanBeNative(const FName MetaDataName)
+{
+	if (MetaDataName == TEXT("ModuleRelativePath"))
+	{
+		return false;
+	}
+	return true;
+}
+
+FString FEmitHelper::HandleMetaData(const UField* Field, bool AddCategory, TArray<FString>* AdditinalMetaData)
+{
+	FString MetaDataStr;
+
+	check(Field);
+	UPackage* Package = Field->GetOutermost();
+	check(Package);
+	const UMetaData* MetaData = Package->GetMetaData();
+	check(MetaData);
+	const TMap<FName, FString>* ValuesMap = MetaData->ObjectMetaDataMap.Find(Field);
+	TArray<FString> MetaDataStrings;
+	if (ValuesMap && ValuesMap->Num())
+	{
+		for (auto& Pair : *ValuesMap)
+		{
+			if (!MetaDataCanBeNative(Pair.Key))
+			{
+				continue;
+			}
+			if (!Pair.Value.IsEmpty())
+			{
+				FString Value = Pair.Value.Replace(TEXT("\n"), TEXT(""));
+				MetaDataStrings.Emplace(FString::Printf(TEXT("%s=\"%s\""), *Pair.Key.ToString(), *Value));
+			}
+			else
+			{
+				MetaDataStrings.Emplace(Pair.Key.ToString());
+			}
+		}
+	}
+	if (AddCategory && (!ValuesMap || !ValuesMap->Find(TEXT("Category"))))
+	{
+		MetaDataStrings.Emplace(TEXT("Category"));
+	}
+	if (AdditinalMetaData)
+	{
+		MetaDataStrings.Append(*AdditinalMetaData);
+	}
+	MetaDataStrings.Remove(FString());
+	if (MetaDataStrings.Num())
+	{
+		MetaDataStr += TEXT("meta=(");
+		ArrayToString(MetaDataStrings, MetaDataStr, TEXT(", "));
+		MetaDataStr += TEXT(")");
+	}
+	return MetaDataStr;
+}
+
+#ifdef HANDLE_CPF_TAG
+	static_assert(false, "Macro HANDLE_CPF_TAG redefinition.");
+#endif
+#define HANDLE_CPF_TAG(TagName, CheckedFlags) if (HasAllFlags(Flags, (CheckedFlags))) { Tags.Emplace(TagName); }
+
+TArray<FString> FEmitHelper::ProperyFlagsToTags(uint64 Flags)
+{
+	TArray<FString> Tags;
+
+	// EDIT FLAGS
+	if (HasAllFlags(Flags, (CPF_Edit | CPF_EditConst | CPF_DisableEditOnInstance)))
+	{
+		Tags.Emplace(TEXT("VisibleDefaultsOnly"));
+	}
+	else if (HasAllFlags(Flags, (CPF_Edit | CPF_EditConst | CPF_DisableEditOnTemplate)))
+	{
+		Tags.Emplace(TEXT("VisibleInstanceOnly"));
+	}
+	else if (HasAllFlags(Flags, (CPF_Edit | CPF_EditConst)))
+	{
+		Tags.Emplace(TEXT("VisibleAnywhere"));
+	}
+	else if (HasAllFlags(Flags, (CPF_Edit | CPF_DisableEditOnInstance)))
+	{
+		Tags.Emplace(TEXT("EditDefaultsOnly"));
+	}
+	else if (HasAllFlags(Flags, (CPF_Edit | CPF_DisableEditOnTemplate)))
+	{
+		Tags.Emplace(TEXT("EditInstanceOnly"));
+	}
+	else if (HasAllFlags(Flags, (CPF_Edit)))
+	{
+		Tags.Emplace(TEXT("EditAnywhere"));
+	}
+
+	// BLUEPRINT EDIT
+	if (HasAllFlags(Flags, (CPF_BlueprintVisible | CPF_BlueprintReadOnly)))
+	{
+		Tags.Emplace(TEXT("BlueprintReadOnly"));
+	}
+	else if (HasAllFlags(Flags, (CPF_BlueprintVisible)))
+	{
+		Tags.Emplace(TEXT("BlueprintReadWrite"));
+	}
+
+	// CONFIG
+	if (HasAllFlags(Flags, (CPF_GlobalConfig | CPF_Config)))
+	{
+		Tags.Emplace(TEXT("GlobalConfig"));
+	}
+	else if (HasAllFlags(Flags, (CPF_Config)))
+	{
+		Tags.Emplace(TEXT("Config"));
+	}
+
+	// OTHER
+	HANDLE_CPF_TAG(TEXT("Transient"), CPF_Transient)
+	HANDLE_CPF_TAG(TEXT("DuplicateTransient"), CPF_DuplicateTransient)
+	HANDLE_CPF_TAG(TEXT("TextExportTransient"), CPF_TextExportTransient)
+	HANDLE_CPF_TAG(TEXT("NonPIEDuplicateTransient"), CPF_NonPIEDuplicateTransient)
+	HANDLE_CPF_TAG(TEXT("Export"), CPF_ExportObject)
+	HANDLE_CPF_TAG(TEXT("NoClear"), CPF_NoClear)
+	HANDLE_CPF_TAG(TEXT("EditFixedSize"), CPF_EditFixedSize)
+	HANDLE_CPF_TAG(TEXT("NotReplicated"), CPF_RepSkip)
+	HANDLE_CPF_TAG(TEXT("Interp"), CPF_Edit | CPF_BlueprintVisible | CPF_Interp)
+	HANDLE_CPF_TAG(TEXT("NonTransactional"), CPF_NonTransactional)
+	HANDLE_CPF_TAG(TEXT("BlueprintAssignable"), CPF_BlueprintAssignable)
+	HANDLE_CPF_TAG(TEXT("BlueprintCallable"), CPF_BlueprintCallable)
+	HANDLE_CPF_TAG(TEXT("BlueprintAuthorityOnly"), CPF_BlueprintAuthorityOnly)
+	HANDLE_CPF_TAG(TEXT("AssetRegistrySearchable"), CPF_AssetRegistrySearchable)
+	HANDLE_CPF_TAG(TEXT("SimpleDisplay"), CPF_SimpleDisplay)
+	HANDLE_CPF_TAG(TEXT("AdvancedDisplay"), CPF_AdvancedDisplay)
+	HANDLE_CPF_TAG(TEXT("SaveGame"), CPF_SaveGame)
+
+	//TODO:
+	//HANDLE_CPF_TAG(TEXT("Instanced"), CPF_PersistentInstance | CPF_ExportObject | CPF_InstancedReference)
+
+	return Tags;
+}
+
+TArray<FString> FEmitHelper::FunctionFlagsToTags(uint64 Flags)
+{
+	TArray<FString> Tags;
+
+	//Pointless: BlueprintNativeEvent, BlueprintImplementableEvent
+	//Pointless: CustomThunk
+
+	//TODO: SealedEvent
+	//TODO: Unreliable
+	//TODO: ServiceRequest, ServiceResponse
+
+	HANDLE_CPF_TAG(TEXT("Exec"), FUNC_Exec)
+	HANDLE_CPF_TAG(TEXT("Server"), FUNC_Net | FUNC_NetServer)
+	HANDLE_CPF_TAG(TEXT("Client"), FUNC_Net | FUNC_NetClient)
+	HANDLE_CPF_TAG(TEXT("NetMulticast"), FUNC_Net | FUNC_NetMulticast)
+	HANDLE_CPF_TAG(TEXT("Reliable"), FUNC_NetReliable)
+	HANDLE_CPF_TAG(TEXT("BlueprintCallable"), FUNC_BlueprintCallable)
+	HANDLE_CPF_TAG(TEXT("BlueprintPure"), FUNC_BlueprintCallable | FUNC_BlueprintPure)
+	HANDLE_CPF_TAG(TEXT("BlueprintAuthorityOnly"), FUNC_BlueprintAuthorityOnly)
+	HANDLE_CPF_TAG(TEXT("BlueprintCosmetic"), FUNC_BlueprintCosmetic)
+	HANDLE_CPF_TAG(TEXT("WithValidation"), FUNC_NetValidate)
+
+	return Tags;
+}
+
+#undef HANDLE_CPF_TAG
+
+bool FEmitHelper::IsBlueprintNativeEvent(uint64 FunctionFlags)
+{
+	return HasAllFlags(FunctionFlags, FUNC_Event | FUNC_BlueprintEvent | FUNC_Native);
+}
+
+bool FEmitHelper::IsBlueprintImplementableEvent(uint64 FunctionFlags)
+{
+	return HasAllFlags(FunctionFlags, FUNC_Event | FUNC_BlueprintEvent) && !HasAllFlags(FunctionFlags, FUNC_Native);
+}
+
+FString FEmitHelper::EmitUFuntion(UFunction* Function, TArray<FString>* AdditinalMetaData)
+{
+	TArray<FString> Tags = FEmitHelper::FunctionFlagsToTags(Function->FunctionFlags);
+
+	auto FunctionOwnerClass = Function->GetOwnerClass();
+	if (FunctionOwnerClass->IsChildOf<UInterface>())
+	{
+		Tags.Emplace(TEXT("BlueprintNativeEvent"));
+	}
+
+	const bool bMustHaveCategory = (Function->FunctionFlags & (FUNC_BlueprintCallable | FUNC_BlueprintPure)) != 0;
+	Tags.Emplace(FEmitHelper::HandleMetaData(Function, bMustHaveCategory, AdditinalMetaData));
+	Tags.Remove(FString());
+
+	FString AllTags;
+	FEmitHelper::ArrayToString(Tags, AllTags, TEXT(", "));
+
+	return FString::Printf(TEXT("UFUNCTION(%s)"), *AllTags);
+}
+
+int32 FEmitHelper::ParseDelegateDetails(UFunction* Signature, FString& OutParametersMacro, FString& OutParamNumberStr)
+{
+	int32 ParameterNum = 0;
+	FStringOutputDevice Parameters;
+	for (TFieldIterator<UProperty> PropIt(Signature); PropIt && (PropIt->PropertyFlags & CPF_Parm); ++PropIt)
+	{
+		Parameters += ", ";
+		PropIt->ExportCppDeclaration(Parameters, EExportedDeclaration::MacroParameter, NULL, EPropertyExportCPPFlags::CPPF_CustomTypeName);
+		++ParameterNum;
+	}
+
+	FString ParamNumberStr;
+	switch (ParameterNum)
+	{
+	case 0: ParamNumberStr = TEXT("");				break;
+	case 1: ParamNumberStr = TEXT("_OneParam");		break;
+	case 2: ParamNumberStr = TEXT("_TwoParams");	break;
+	case 3: ParamNumberStr = TEXT("_ThreeParams");	break;
+	case 4: ParamNumberStr = TEXT("_FourParams");	break;
+	case 5: ParamNumberStr = TEXT("_FiveParams");	break;
+	case 6: ParamNumberStr = TEXT("_SixParams");	break;
+	case 7: ParamNumberStr = TEXT("_SevenParams");	break;
+	case 8: ParamNumberStr = TEXT("_EightParams");	break;
+	default: ParamNumberStr = TEXT("_TooMany");		break;
+	}
+
+	OutParametersMacro = Parameters;
+	OutParamNumberStr = ParamNumberStr;
+	return ParameterNum;
+}
+
+TArray<FString> FEmitHelper::EmitSinglecastDelegateDeclarations(const TArray<UDelegateProperty*>& Delegates)
+{
+	TArray<FString> Results;
+	for (auto It : Delegates)
+	{
+		check(It);
+		auto Signature = It->SignatureFunction;
+		check(Signature);
+
+		FString ParamNumberStr, Parameters;
+		ParseDelegateDetails(Signature, Parameters, ParamNumberStr);
+
+		Results.Add(*FString::Printf(TEXT("DECLARE_DYNAMIC_DELEGATE%s(%s%s)"),
+			*ParamNumberStr, *It->GetCPPType(NULL, EPropertyExportCPPFlags::CPPF_CustomTypeName), *Parameters));
+	}
+	return Results;
+}
+
+TArray<FString> FEmitHelper::EmitMulticastDelegateDeclarations(UClass* SourceClass)
+{
+	TArray<FString> Results;
+	for (TFieldIterator<UMulticastDelegateProperty> It(SourceClass, EFieldIteratorFlags::ExcludeSuper); It; ++It)
+	{
+		auto Signature = It->SignatureFunction;
+		check(Signature);
+
+		FString ParamNumberStr, Parameters;
+		ParseDelegateDetails(Signature, Parameters, ParamNumberStr);
+
+		Results.Add(*FString::Printf(TEXT("DECLARE_DYNAMIC_MULTICAST_DELEGATE%s(%s%s)"),
+			*ParamNumberStr, *It->GetCPPType(NULL, EPropertyExportCPPFlags::CPPF_CustomTypeName), *Parameters));
+	}
+
+	return Results;
+}
+
+FString FEmitHelper::EmitLifetimeReplicatedPropsImpl(UClass* SourceClass, const FString& CppClassName, const TCHAR* InCurrentIndent)
+{
+	FString Result;
+	bool bFunctionInitilzed = false;
+	for (TFieldIterator<UProperty> It(SourceClass, EFieldIteratorFlags::ExcludeSuper); It; ++It)
+	{
+		if ((It->PropertyFlags & CPF_Net) != 0)
+		{
+			if (!bFunctionInitilzed)
+			{
+				Result += FString::Printf(TEXT("%svoid %s::GetLifetimeReplicatedProps(TArray< FLifetimeProperty > & OutLifetimeProps) const\n"), InCurrentIndent, *CppClassName);
+				Result += FString::Printf(TEXT("%s{\n"), InCurrentIndent);
+				Result += FString::Printf(TEXT("%s\tSuper::GetLifetimeReplicatedProps(OutLifetimeProps);\n"), InCurrentIndent);
+				bFunctionInitilzed = true;
+			}
+			Result += FString::Printf(TEXT("%s\tDOREPLIFETIME( %s, %s);\n"), InCurrentIndent, *CppClassName, *It->GetNameCPP());
+		}
+	}
+	if (bFunctionInitilzed)
+	{
+		Result += FString::Printf(TEXT("%s}\n"), InCurrentIndent);
+	}
+	return Result;
+}
+
+FString FEmitHelper::GatherNativeHeadersToInclude(UField* SourceItem, const TArray<FString>& PersistentHeaders)
+{
+	TSet<FString> HeaderFiles;
+	HeaderFiles.Append(PersistentHeaders);
+	{
+		TArray<UObject*> ReferencedObjects;
+		{
+			FReferenceFinder ReferenceFinder(ReferencedObjects, NULL, false, false, false, true);
+			TArray<UObject*> ObjectsToCheck;
+			GetObjectsWithOuter(SourceItem, ObjectsToCheck, true);
+			//CDO?
+			for (auto Obj : ObjectsToCheck)
+			{
+				ReferenceFinder.FindReferences(Obj);
+			}
+		}
+
+		auto EngineSourceDir = FPaths::EngineSourceDir();
+		auto GameSourceDir = FPaths::GameSourceDir();
+		auto CurrentPackage = SourceItem->GetTypedOuter<UPackage>();
+		for (auto Obj : ReferencedObjects)
+		{
+			auto Field = Cast<UField>(Obj);
+			FString PackPath;
+			if (Field
+				&& (Field->GetTypedOuter<UPackage>() != CurrentPackage)
+				&& FSourceCodeNavigation::FindClassHeaderPath(Field, PackPath))
+			{
+				if (!PackPath.RemoveFromStart(EngineSourceDir))
+				{
+					if (!PackPath.RemoveFromStart(GameSourceDir))
+					{
+						PackPath = FPaths::GetCleanFilename(PackPath);
+					}
+				}
+				HeaderFiles.Add(PackPath);
+			}
+		}
+	}
+
+	FString Result;
+	for (auto HeaderFile : HeaderFiles)
+	{
+		Result += FString::Printf(TEXT("#include \"%s\"\n"), *HeaderFile);
+	}
+
+	return Result;
+}
+
+FString FEmitHelper::LiteralTerm(const FEdGraphPinType& Type, const FString& CustomValue, UObject* LiteralObject)
+{
+	auto Schema = GetDefault<UEdGraphSchema_K2>();
+
+	if (UEdGraphSchema_K2::PC_String == Type.PinCategory)
+	{
+		return FString::Printf(TEXT("TEXT(\"%s\")"), *CustomValue.ReplaceCharWithEscapedChar());
+	}
+	else if (UEdGraphSchema_K2::PC_Text == Type.PinCategory)
+	{
+		return FString::Printf(TEXT("FText::FromString(TEXT(\"%s\"))"), *CustomValue.ReplaceCharWithEscapedChar());
+	}
+	else if (UEdGraphSchema_K2::PC_Float == Type.PinCategory)
+	{
+		float Value = CustomValue.IsEmpty() ? 0.0f : FCString::Atof(*CustomValue);
+		return FString::Printf(TEXT("%f"), Value);
+	}
+	else if (UEdGraphSchema_K2::PC_Int == Type.PinCategory)
+	{
+		int32 Value = CustomValue.IsEmpty() ? 0 : FCString::Atoi(*CustomValue);
+		return FString::Printf(TEXT("%d"), Value);
+	}
+	else if ((UEdGraphSchema_K2::PC_Byte == Type.PinCategory) || (UEdGraphSchema_K2::PC_Enum == Type.PinCategory))
+	{
+		auto TypeEnum = Cast<UEnum>(Type.PinSubCategoryObject.Get());
+		if (TypeEnum)
+		{
+			return FString::Printf(TEXT("%s::%s"), *TypeEnum->GetName(), CustomValue.IsEmpty() 
+				? *TypeEnum->GetEnumName(TypeEnum->NumEnums()-1)
+				: *CustomValue);
+		}
+		else
+		{
+			uint8 Value = CustomValue.IsEmpty() ? 0 : FCString::Atoi(*CustomValue);
+			return FString::Printf(TEXT("%u"), Value);
+		}
+	}
+	else if (UEdGraphSchema_K2::PC_Boolean == Type.PinCategory)
+	{
+		const bool bValue = CustomValue.ToBool();
+		return bValue ? TEXT("true") : TEXT("false");
+	}
+	else if (UEdGraphSchema_K2::PC_Name == Type.PinCategory)
+	{
+		return CustomValue.IsEmpty() 
+			? TEXT("FName()") 
+			: FString::Printf(TEXT("FName(TEXT(\"%s\"))"), *(FName(*CustomValue).ToString().ReplaceCharWithEscapedChar()));
+	}
+	else if (UEdGraphSchema_K2::PC_Struct == Type.PinCategory)
+	{
+		auto StructType = Cast<UScriptStruct>(Type.PinSubCategoryObject.Get());
+		ensure(StructType);
+		if (StructType == TBaseStructure<FVector>::Get())
+		{
+			FVector Vect = FVector::ZeroVector;
+			FDefaultValueHelper::ParseVector(CustomValue, /*out*/ Vect);
+			return FString::Printf(TEXT("FVector(%f,%f,%f)"), Vect.X, Vect.Y, Vect.Z);
+		}
+		else if (StructType == TBaseStructure<FRotator>::Get())
+		{
+			FRotator Rot = FRotator::ZeroRotator;
+			FDefaultValueHelper::ParseRotator(CustomValue, /*out*/ Rot);
+			return FString::Printf(TEXT("FRotator(%f,%f,%f)"), Rot.Pitch, Rot.Yaw, Rot.Roll);
+		}
+		else if (StructType == TBaseStructure<FTransform>::Get())
+		{
+			FTransform Trans = FTransform::Identity;
+			Trans.InitFromString(CustomValue);
+			const FQuat Rot = Trans.GetRotation();
+			const FVector Translation = Trans.GetTranslation();
+			const FVector Scale = Trans.GetScale3D();
+			return FString::Printf(TEXT("FTransform( FQuat(%f,%f,%f,%f), FVector(%f,%f,%f), FVector(%f,%f,%f) )"),
+				Rot.X, Rot.Y, Rot.Z, Rot.W, Translation.X, Translation.Y, Translation.Z, Scale.X, Scale.Y, Scale.Z);
+		}
+		else if (StructType == TBaseStructure<FLinearColor>::Get())
+		{
+			FLinearColor LinearColor;
+			LinearColor.InitFromString(CustomValue);
+			return FString::Printf(TEXT("FLinearColor(%f,%f,%f,%f)"), LinearColor.R, LinearColor.G, LinearColor.B, LinearColor.A);
+		}
+		else if (StructType == TBaseStructure<FColor>::Get())
+		{
+			FColor Color;
+			Color.InitFromString(CustomValue);
+			return FString::Printf(TEXT("FColor(%d,%d,%d,%d)"), Color.R, Color.G, Color.B, Color.A);
+		}
+		if (StructType == TBaseStructure<FVector2D>::Get())
+		{
+			FVector2D Vect = FVector2D::ZeroVector;
+			Vect.InitFromString(CustomValue);
+			return FString::Printf(TEXT("FVector2D(%f,%f)"), Vect.X, Vect.Y);
+		}
+		else
+		{
+			//@todo:  This needs to be more robust, since import text isn't really proper for struct construction.
+			return FString(TEXT("F")) + StructType->GetName() + (CustomValue.IsEmpty() ? TEXT("{}") : *CustomValue);
+		}
+	}
+	else if (Type.PinSubCategory == UEdGraphSchema_K2::PSC_Self)
+	{
+		return TEXT("this");
+	}
+	else if (UEdGraphSchema_K2::PC_Class == Type.PinCategory)
+	{
+		if (auto FoundClass = Cast<const UClass>(LiteralObject))
+		{
+			return FString::Printf(TEXT("%s%s::StaticClass()"), FoundClass->GetPrefixCPP(), *FoundClass->GetName());
+		}
+		return FString(TEXT("nullptr"));
+	}
+	else if((UEdGraphSchema_K2::PC_AssetClass == Type.PinCategory) || (UEdGraphSchema_K2::PC_Asset == Type.PinCategory))
+	{
+		return LiteralObject
+			? FString::Printf(TEXT("FStringAssetReference(TEXT(\"%s\"))"), *(LiteralObject->GetPathName().ReplaceCharWithEscapedChar()))
+			: FString(TEXT("nullptr"));
+	}
+	else if (UEdGraphSchema_K2::PC_Object == Type.PinCategory)
+	{
+		if (LiteralObject)
+		{
+			UClass* FoundClass = Cast<UClass>(Type.PinSubCategoryObject.Get());
+			FString ClassString = FoundClass ? (FString(FoundClass->GetPrefixCPP()) + FoundClass->GetName()) : TEXT("UObject");
+			return FString::Printf(TEXT("LoadObject<%s>(nullptr, TEXT(\"%s\"))"), *ClassString, *(LiteralObject->GetPathName().ReplaceCharWithEscapedChar()));
+		}
+		else
+		{
+			return FString(TEXT("nullptr"));
+		}
+	}
+	/*
+	else if (CoerceProperty->IsA(UInterfaceProperty::StaticClass()))
+	{
+		if (Term->Type.PinSubCategory == PSC_Self)
+		{
+			return TEXT("this");
+		}
+		else
+		{
+			ensureMsgf(false, TEXT("It is not possible to express this interface property as a literal value!"));
+			return Term->Name;
+		}
+	}
+	*/
+	ensureMsgf(false, TEXT("It is not possible to express this type as a literal value!"));
+	return CustomValue;
+}
+
+FString FEmitHelper::DefaultValue(const FEdGraphPinType& Type)
+{
+	return LiteralTerm(Type, FString(), nullptr);
+}
+
+UFunction* FEmitHelper::GetOriginalFunction(UFunction* Function)
+{
+	check(Function);
+	const FName FunctionName = Function->GetFName();
+
+	UClass* Owner = Function->GetOwnerClass();
+	check(Owner);
+	for (auto& Inter : Owner->Interfaces)
+	{
+		if (UFunction* Result = Inter.Class->FindFunctionByName(FunctionName))
+		{
+			return GetOriginalFunction(Result);
+		}
+	}
+
+	for (UClass* SearchClass = Owner->GetSuperClass(); SearchClass; SearchClass = SearchClass->GetSuperClass())
+	{
+		if (UFunction* Result = SearchClass->FindFunctionByName(FunctionName))
+		{
+			return GetOriginalFunction(Result);
+		}
+	}
+
+	return Function;
+}
+
+bool FEmitHelper::ShoulsHandleAsNativeEvent(UFunction* Function)
+{
+	check(Function);
+	auto OriginalFunction = FEmitHelper::GetOriginalFunction(Function);
+	check(OriginalFunction);
+	if (OriginalFunction != Function)
+	{
+		const uint32 FlagsToCheckMask = EFunctionFlags::FUNC_Event | EFunctionFlags::FUNC_BlueprintEvent | EFunctionFlags::FUNC_Native;
+		const uint32 FlagsToCheck = OriginalFunction->FunctionFlags & FlagsToCheckMask;
+
+		auto OriginalOwner = OriginalFunction->GetOwnerClass();
+		const bool bBPInterface = Cast<UBlueprintGeneratedClass>(OriginalOwner) && OriginalOwner->IsChildOf<UInterface>();
+		return (FlagsToCheck == FlagsToCheckMask) || bBPInterface;
+	}
+	return false;
+}
+
+bool FEmitHelper::ShoulsHandleAsImplementableEvent(UFunction* Function)
+{
+	check(Function);
+	auto OriginalFunction = FEmitHelper::GetOriginalFunction(Function);
+	check(OriginalFunction);
+	if (OriginalFunction != Function)
+	{
+		const uint32 FlagsToCheckMask = EFunctionFlags::FUNC_Event | EFunctionFlags::FUNC_BlueprintEvent | EFunctionFlags::FUNC_Native;
+		const uint32 FlagsToCheck = OriginalFunction->FunctionFlags & FlagsToCheckMask;
+		return (FlagsToCheck == (EFunctionFlags::FUNC_Event | EFunctionFlags::FUNC_BlueprintEvent));
+	}
+	return false;
+}
+
+FString FEmitDefaultValueHelper::GenerateGetDefaultValue(const UUserDefinedStruct* Struct)
+{
+	check(Struct);
+	const FString StructName = FString(TEXT("F")) + Struct->GetName();
+	FString Result = FString::Printf(TEXT("\tstatic %s GetDefaultValue()\n\t{\n\t\t%s DefaultData__;\n"), *StructName, *StructName);
+
+	FStructOnScope StructData(Struct);
+	FStructureEditorUtils::Fill_MakeStructureDefaultValue(Struct, StructData.GetStructMemory());
+
+	for (auto Property : TFieldRange<const UProperty>(Struct))
+	{
+		const bool bStaticArray = (Property->ArrayDim > 1);
+		for (int32 ArrayIndex = 0; ArrayIndex < Property->ArrayDim; ++ArrayIndex)
+		{
+			const uint8* ValuePtr = Property->ContainerPtrToValuePtr<uint8>(StructData.GetStructMemory(), ArrayIndex);
+			const FString ArrayPost = bStaticArray ? FString::Printf(TEXT("[%d]"), ArrayIndex) : TEXT("");
+			const FString PathToMember = FString::Printf(TEXT("\t\tDefaultData__.%s%s"), *Property->GetNameCPP(), *ArrayPost);
+			InnerGenerate(Property, ValuePtr, PathToMember, Result);
+		}
+	}
+
+	Result += TEXT("\n\t\treturn DefaultData__;\n\t}\n");
+	return Result;
+}
+
+void FEmitDefaultValueHelper::InnerGenerate(const UProperty* Property, const uint8* ValuePtr, const FString& PathToMember, FString& OutResult)
+{
+	if (!Property->GetOuter()->IsA<UArrayProperty>())
+	{
+		FString ValueStr;
+		Property->ExportTextItem(ValueStr, ValuePtr, ValuePtr, nullptr, EPropertyPortFlags::PPF_ExportCpp);
+		OutResult += FString::Printf(TEXT("%s = %s;\n"), *PathToMember, *ValueStr);
+	}
+
+	if (auto StructProperty = Cast<const UStructProperty>(Property))
+	{
+		for (auto LocalProperty : TFieldRange<const UProperty>(StructProperty->Struct))
+		{
+			const bool bStaticArray = (Property->ArrayDim > 1);
+			for (int32 ArrayIndex = 0; ArrayIndex < Property->ArrayDim; ++ArrayIndex)
+			{
+				const uint8* LocalValuePtr = LocalProperty->ContainerPtrToValuePtr<uint8>(ValuePtr, ArrayIndex);
+				const FString ArrayPost = bStaticArray ? FString::Printf(TEXT("[%d]"), ArrayIndex) : TEXT("");
+				const FString LocalPathToMember = FString::Printf(TEXT("%s.%s%s"), *PathToMember, *LocalProperty->GetNameCPP(),* ArrayPost);
+				InnerGenerate(LocalProperty, LocalValuePtr, LocalPathToMember, OutResult);
+			}
+		}
+	}
+	else if (auto ArrayProperty = Cast<const UArrayProperty>(Property))
+	{
+		check(ArrayProperty->Inner);
+		FScriptArrayHelper ScriptArrayHelper(ArrayProperty, ValuePtr);
+		for (int32 Index = 0; Index < ScriptArrayHelper.Num(); ++Index)
+		{
+			const uint8* LocalValuePtr = ScriptArrayHelper.GetRawPtr(Index);
+			{
+				FString ValueStr;
+				ArrayProperty->Inner->ExportTextItem(ValueStr, LocalValuePtr, LocalValuePtr, nullptr, EPropertyPortFlags::PPF_ExportCpp);
+				OutResult += FString::Printf(TEXT("%s.Add(%s);\n"), *PathToMember, *ValueStr);
+			}
+			const FString LocalPathToMember = FString::Printf(TEXT("%s[%d]"), *PathToMember, Index);
+			InnerGenerate(ArrayProperty->Inner, LocalValuePtr, LocalPathToMember, OutResult);
+		}
+	}
+}
