Fixed non unity build error

diff --git a/Engine/Source/Developer/BlueprintCompilerCppBackend/Private/BlueprintCompilerCppBackendUtils.cpp b/Engine/Source/Developer/BlueprintCompilerCppBackend/Private/BlueprintCompilerCppBackendUtils.cpp
--- a/Engine/Source/Developer/BlueprintCompilerCppBackend/Private/BlueprintCompilerCppBackendUtils.cpp
+++ b/Engine/Source/Developer/BlueprintCompilerCppBackend/Private/BlueprintCompilerCppBackendUtils.cpp
@@ -2,6 +2,7 @@
 
 #include "BlueprintCompilerCppBackendModulePrivatePCH.h"
 #include "BlueprintCompilerCppBackendUtils.h"
+#include "BlueprintCompilerCppBackend.h"
 #include "EdGraphSchema_K2.h"
 #include "Editor/UnrealEd/Public/Kismet2/StructureEditorUtils.h"
 
