UEBP-40 C++ code generation (WIP):
- Enum/Byte conversion
- FEmitDefaultValueHelper improvements
- Object constructor - early version (no subobjects/components)

diff --git a/Engine/Source/Developer/BlueprintCompilerCppBackend/Private/BlueprintCompilerCppBackend.cpp b/Engine/Source/Developer/BlueprintCompilerCppBackend/Private/BlueprintCompilerCppBackend.cpp
--- a/Engine/Source/Developer/BlueprintCompilerCppBackend/Private/BlueprintCompilerCppBackend.cpp
+++ b/Engine/Source/Developer/BlueprintCompilerCppBackend/Private/BlueprintCompilerCppBackend.cpp
@@ -204,8 +204,11 @@ void FBlueprintCompilerCppBackend::GenerateCodeFromClass(UClass* SourceClass, TI
 		Emit(Header, TEXT("\n"));
 	}
 
-	//constructor is generated by GENERATED_BODY macro
-	// Emit(Body, *FString::Printf(TEXT("%s::%s(const FObjectInitializer& ObjectInitializer) : Super(ObjectInitializer) {}\n\n"), *CppClassName, *CppClassName));
+	if (!bIsInterface)
+	{
+		Emit(Header, *FString::Printf(TEXT("\t%s(const FObjectInitializer& ObjectInitializer);\n\n"), *CppClassName));
+		Emit(Body, *FEmitDefaultValueHelper::GenerateConstructor(SourceClass));
+	}
 
 	for (int32 i = 0; i < Functions.Num(); ++i)
 	{
@@ -405,14 +408,19 @@ FString FBlueprintCompilerCppBackend::EmitSwitchValueStatmentInner(FBlueprintCom
 
 void FBlueprintCompilerCppBackend::EmitAssignmentStatment(FKismetFunctionContext& FunctionContext, FBlueprintCompiledStatement& Statement)
 {
+	check(Statement.LHS && Statement.RHS[0]);
 	FString DestinationExpression = TermToText(Statement.LHS, false);
 	FString SourceExpression = TermToText(Statement.RHS[0]);
 
 	FSafeContextScopedEmmitter SafeContextScope(Body, Statement.LHS->Context, *this, TEXT("\t\t\t"));
+
+	FString BeginCast;
+	FString EndCast;
+	FEmitHelper::GenerateAssignmentCast(Statement.LHS->Type, Statement.RHS[0]->Type, BeginCast, EndCast);
 	// Emit the assignment statement
 	Emit(Body, TEXT("\t\t\t"));
 	Emit(Body, *SafeContextScope.GetAdditionalIndent());
-	Emit(Body, *FString::Printf(TEXT("%s = %s;\n"), *DestinationExpression, *SourceExpression));
+	Emit(Body, *FString::Printf(TEXT("%s = %s%s%s;\n"), *DestinationExpression, *BeginCast, *SourceExpression, *EndCast));
 }
 
 void FBlueprintCompilerCppBackend::EmitCastObjToInterfaceStatement(FKismetFunctionContext& FunctionContext, FBlueprintCompiledStatement& Statement)
diff --git a/Engine/Source/Developer/BlueprintCompilerCppBackend/Private/BlueprintCompilerCppBackendUtils.cpp b/Engine/Source/Developer/BlueprintCompilerCppBackend/Private/BlueprintCompilerCppBackendUtils.cpp
--- a/Engine/Source/Developer/BlueprintCompilerCppBackend/Private/BlueprintCompilerCppBackendUtils.cpp
+++ b/Engine/Source/Developer/BlueprintCompilerCppBackend/Private/BlueprintCompilerCppBackendUtils.cpp
@@ -651,6 +651,23 @@ bool FEmitHelper::ShoulsHandleAsImplementableEvent(UFunction* Function)
 	return false;
 }
 
+bool FEmitHelper::GenerateAssignmentCast(const FEdGraphPinType& LType, const FEdGraphPinType& RType, FString& OutCastBegin, FString& OutCastEnd)
+{
+	// BYTE to ENUM cast
+	auto LTypeEnum = Cast<UEnum>(LType.PinSubCategoryObject.Get());
+	if ((LType.PinCategory == UEdGraphSchema_K2::PC_Byte)
+		&& (RType.PinCategory == UEdGraphSchema_K2::PC_Byte)
+		&& !RType.PinSubCategoryObject.IsValid()
+		&& LTypeEnum)
+	{
+		FString EnumCppType = !LTypeEnum->CppType.IsEmpty() ? LTypeEnum->CppType : LTypeEnum->GetName();
+		OutCastBegin = FString::Printf(TEXT("static_cast<%s>("), *EnumCppType);
+		OutCastEnd = TEXT(")");
+		return true;
+	}
+	return false;
+}
+
 FString FEmitDefaultValueHelper::GenerateGetDefaultValue(const UUserDefinedStruct* Struct)
 {
 	check(Struct);
@@ -732,7 +749,7 @@ void FEmitDefaultValueHelper::InnerGenerate(const UProperty* Property, const uin
 			{
 				const uint8* LocalValuePtr = LocalProperty->ContainerPtrToValuePtr<uint8>(ValuePtr, ArrayIndex);
 				const FString ArrayPost = bStaticArray ? FString::Printf(TEXT("[%d]"), ArrayIndex) : TEXT("");
-				const FString LocalPathToMember = FString::Printf(TEXT("%s.%s%s"), *PathToMember, *LocalProperty->GetNameCPP(),* ArrayPost);
+				const FString LocalPathToMember = FString::Printf(TEXT("%s.%s%s"), *PathToMember, *LocalProperty->GetNameCPP(), *ArrayPost);
 				InnerGenerate(LocalProperty, LocalValuePtr, LocalPathToMember, OutResult);
 			}
 		}
@@ -759,6 +776,8 @@ void FEmitDefaultValueHelper::InnerGenerate(const UProperty* Property, const uin
 
 bool FEmitDefaultValueHelper::HandleSpecialTypes(const UProperty* Property, const uint8* ValuePtr, FString& OutResult)
 {
+	//TODO: Use Path maps for Objects
+
 	if (auto StructProperty = Cast<UStructProperty>(Property))
 	{
 		if (TBaseStructure<FTransform>::Get() == StructProperty->Struct)
@@ -774,6 +793,49 @@ bool FEmitDefaultValueHelper::HandleSpecialTypes(const UProperty* Property, cons
 				, Scale.X, Scale.Y, Scale.Z);
 			return true;
 		}
+
+		if (TBaseStructure<FVector>::Get() == StructProperty->Struct)
+		{
+			const FVector* Vector = reinterpret_cast<const FVector*>(ValuePtr);
+			OutResult = FString::Printf(TEXT("FVector(%f, %f, %f)"), Vector->X, Vector->Y, Vector->Z);
+			return true;
+		}
 	}
 	return false;
+}
+
+FString FEmitDefaultValueHelper::GenerateConstructor(UClass* BPGC)
+{
+	check(BPGC);
+	const FString CppClassName = FString(BPGC->GetPrefixCPP()) + BPGC->GetName();
+
+	FString Result;
+	Result += FString::Printf(TEXT("%s::%s(const FObjectInitializer& ObjectInitializer) : Super(ObjectInitializer)\n{\n"), *CppClassName, *CppClassName);
+
+	UObject* CDO = BPGC->GetDefaultObject(false);
+	UObject* ParentCDO = BPGC->GetSuperClass()->GetDefaultObject(false);
+	check(CDO && ParentCDO);
+	for (auto Property : TFieldRange<const UProperty>(BPGC))
+	{
+		const bool bNewProperty = Property->GetOwnerStruct() == BPGC;
+		const bool bIsAccessible = bNewProperty || !Property->HasAnyPropertyFlags(CPF_NativeAccessSpecifierPrivate);
+		if (bIsAccessible)
+		{
+			const bool bStaticArray = (Property->ArrayDim > 1);
+			for (int32 ArrayIndex = 0; ArrayIndex < Property->ArrayDim; ++ArrayIndex)
+			{
+				const bool bChangedValue = !bNewProperty && !Property->Identical_InContainer(CDO, ParentCDO, ArrayIndex);
+				if (bNewProperty || bChangedValue)
+				{
+					const uint8* ValuePtr = Property->ContainerPtrToValuePtr<uint8>(CDO, ArrayIndex);
+					const FString ArrayPost = bStaticArray ? FString::Printf(TEXT("[%d]"), ArrayIndex) : TEXT("");
+					const FString PathToMember = FString::Printf(TEXT("\t%s"), *Property->GetNameCPP(), *ArrayPost);
+					InnerGenerate(Property, ValuePtr, PathToMember, Result);
+				}
+			}
+		}
+	}
+
+	Result += TEXT("\n}\n");
+	return Result;
 }
\ No newline at end of file
diff --git a/Engine/Source/Developer/BlueprintCompilerCppBackend/Private/BlueprintCompilerCppBackendUtils.h b/Engine/Source/Developer/BlueprintCompilerCppBackend/Private/BlueprintCompilerCppBackendUtils.h
--- a/Engine/Source/Developer/BlueprintCompilerCppBackend/Private/BlueprintCompilerCppBackendUtils.h
+++ b/Engine/Source/Developer/BlueprintCompilerCppBackend/Private/BlueprintCompilerCppBackendUtils.h
@@ -61,11 +61,16 @@ struct FEmitHelper
 	static bool ShoulsHandleAsNativeEvent(UFunction* Function);
 
 	static bool ShoulsHandleAsImplementableEvent(UFunction* Function);
+
+	static bool GenerateAssignmentCast(const FEdGraphPinType& LType, const FEdGraphPinType& RType, FString& OutCastBegin, FString& OutCastEnd);
 };
 
 struct FEmitDefaultValueHelper
 {
 	static FString GenerateGetDefaultValue(const UUserDefinedStruct* Struct);
+
+	static FString GenerateConstructor(UClass* BPGC);
+private:
 	static void InnerGenerate(const UProperty* Property, const uint8* ValuePtr, const FString& PathToMember, FString& OutResult);
 	
 	// Return if it's complete construction
