UEBP-40 C++ code generation (WIP):
- Generating Components!
- FEmitDefaultValueHelperimprovements
- KCST_GotoReturn fixed
- constructor fixed
- ProcessUserConstructionScript is called on all actors

diff --git a/Engine/Source/Developer/BlueprintCompilerCppBackend/Private/BlueprintCompilerCppBackend.cpp b/Engine/Source/Developer/BlueprintCompilerCppBackend/Private/BlueprintCompilerCppBackend.cpp
--- a/Engine/Source/Developer/BlueprintCompilerCppBackend/Private/BlueprintCompilerCppBackend.cpp
+++ b/Engine/Source/Developer/BlueprintCompilerCppBackend/Private/BlueprintCompilerCppBackend.cpp
@@ -206,7 +206,7 @@ void FBlueprintCompilerCppBackend::GenerateCodeFromClass(UClass* SourceClass, TI
 
 	if (!bIsInterface)
 	{
-		Emit(Header, *FString::Printf(TEXT("\t%s(const FObjectInitializer& ObjectInitializer);\n\n"), *CppClassName));
+		Emit(Header, *FString::Printf(TEXT("\t%s(const FObjectInitializer& ObjectInitializer = FObjectInitializer::Get());\n\n"), *CppClassName));
 		Emit(Body, *FEmitDefaultValueHelper::GenerateConstructor(SourceClass));
 	}
 
@@ -612,6 +612,7 @@ void FBlueprintCompilerCppBackend::EmitGotoStatement(FKismetFunctionContext& Fun
 	else if (Statement.Type == KCST_GotoReturn)
 	{
 		Emit(Body, TEXT("\t\t\tCurrentState = -1;\n"));
+		Emit(Body, *FString::Printf(TEXT("\t\t\tbreak;\n")));
 	}
 	else
 	{
diff --git a/Engine/Source/Developer/BlueprintCompilerCppBackend/Private/BlueprintCompilerCppBackendUtils.cpp b/Engine/Source/Developer/BlueprintCompilerCppBackend/Private/BlueprintCompilerCppBackendUtils.cpp
--- a/Engine/Source/Developer/BlueprintCompilerCppBackend/Private/BlueprintCompilerCppBackendUtils.cpp
+++ b/Engine/Source/Developer/BlueprintCompilerCppBackend/Private/BlueprintCompilerCppBackendUtils.cpp
@@ -5,6 +5,7 @@
 #include "BlueprintCompilerCppBackend.h"
 #include "EdGraphSchema_K2.h"
 #include "Editor/UnrealEd/Public/Kismet2/StructureEditorUtils.h"
+#include "Engine/InheritableComponentHandler.h"
 
 FString FSafeContextScopedEmmitter::GetAdditionalIndent() const
 {
@@ -668,6 +669,33 @@ bool FEmitHelper::GenerateAssignmentCast(const FEdGraphPinType& LType, const FEd
 	return false;
 }
 
+void FEmitDefaultValueHelper::OuterGenerate(const UProperty* Property, const uint8* DataContainer, const FString OuterPath, const uint8* OptionalDefaultDataContainer, FString& OutResult, bool bAllowProtected)
+{
+	if (Property->HasAnyPropertyFlags(CPF_EditorOnly))
+	{
+		UE_LOG(LogK2Compiler, Log, TEXT("FEmitDefaultValueHelper Skip EditorOnly property: %s"), *Property->GetPathName());
+		return;
+	}
+
+	if (Property->HasAnyPropertyFlags(CPF_NativeAccessSpecifierPrivate) || (!bAllowProtected && Property->HasAnyPropertyFlags(CPF_NativeAccessSpecifierProtected)))
+	{
+		UE_LOG(LogK2Compiler, Warning, TEXT("FEmitDefaultValueHelper Cannot access property: %s"), *Property->GetPathName());
+		return;
+	}
+
+	const bool bStaticArray = (Property->ArrayDim > 1);
+	for (int32 ArrayIndex = 0; ArrayIndex < Property->ArrayDim; ++ArrayIndex)
+	{
+		if (!OptionalDefaultDataContainer || !Property->Identical_InContainer(DataContainer, OptionalDefaultDataContainer, ArrayIndex))
+		{
+			const uint8* ValuePtr = Property->ContainerPtrToValuePtr<uint8>(DataContainer, ArrayIndex);
+			const FString ArrayPost = bStaticArray ? FString::Printf(TEXT("[%d]"), ArrayIndex) : TEXT("");
+			const FString PathToMember = FString::Printf(TEXT("%s%s%s"), *OuterPath, *Property->GetNameCPP(), *ArrayPost);
+			InnerGenerate(Property, ValuePtr, PathToMember, OutResult);
+		}
+	}
+}
+
 FString FEmitDefaultValueHelper::GenerateGetDefaultValue(const UUserDefinedStruct* Struct)
 {
 	check(Struct);
@@ -679,50 +707,37 @@ FString FEmitDefaultValueHelper::GenerateGetDefaultValue(const UUserDefinedStruc
 
 	for (auto Property : TFieldRange<const UProperty>(Struct))
 	{
-		const bool bStaticArray = (Property->ArrayDim > 1);
-		for (int32 ArrayIndex = 0; ArrayIndex < Property->ArrayDim; ++ArrayIndex)
-		{
-			const uint8* ValuePtr = Property->ContainerPtrToValuePtr<uint8>(StructData.GetStructMemory(), ArrayIndex);
-			const FString ArrayPost = bStaticArray ? FString::Printf(TEXT("[%d]"), ArrayIndex) : TEXT("");
-			const FString PathToMember = FString::Printf(TEXT("\t\tDefaultData__.%s%s"), *Property->GetNameCPP(), *ArrayPost);
-			InnerGenerate(Property, ValuePtr, PathToMember, Result);
-		}
+		OuterGenerate(Property, StructData.GetStructMemory(), TEXT("\t\tDefaultData__."), nullptr, Result);
 	}
 
 	Result += TEXT("\n\t\treturn DefaultData__;\n\t}\n");
 	return Result;
 }
 
-bool FEmitDefaultValueHelper::OneLineConstruction(const UProperty* Property, const uint8* ValuePtr, FString& OutResult)
+void FEmitDefaultValueHelper::InnerGenerate(const UProperty* Property, const uint8* ValuePtr, const FString& PathToMember, FString& OutResult)
 {
-	FString ValueStr;
-	bool bComplete = true;
-	if (!HandleSpecialTypes(Property, ValuePtr, ValueStr))
+	auto OneLineConstruction = [](const UProperty* Property, const uint8* ValuePtr, FString& OutResult) -> bool
 	{
-		auto StructProperty = Cast<const UStructProperty>(Property);
-		Property->ExportTextItem(ValueStr, ValuePtr, ValuePtr, nullptr, EPropertyPortFlags::PPF_ExportCpp);
-		if (ValueStr.IsEmpty() && StructProperty)
-		{
-			check(StructProperty->Struct);
-			ValueStr = FString::Printf(TEXT("F%s()"), *StructProperty->Struct->GetName());
-			bComplete = false;
-		}
-		if (ValueStr.IsEmpty())
+		FString ValueStr;
+		bool bComplete = true;
+		if (!HandleSpecialTypes(Property, ValuePtr, ValueStr))
 		{
-			UE_LOG(LogK2Compiler, Warning, TEXT("FEmitDefaultValueHelper Cannot generate initilization: %s"), *Property->GetPathName());
+			auto StructProperty = Cast<const UStructProperty>(Property);
+			Property->ExportTextItem(ValueStr, ValuePtr, ValuePtr, nullptr, EPropertyPortFlags::PPF_ExportCpp);
+			if (ValueStr.IsEmpty() && StructProperty)
+			{
+				check(StructProperty->Struct);
+				ValueStr = FString::Printf(TEXT("F%s()"), *StructProperty->Struct->GetName());
+				bComplete = false;
+			}
+			if (ValueStr.IsEmpty())
+			{
+				UE_LOG(LogK2Compiler, Warning, TEXT("FEmitDefaultValueHelper Cannot generate initilization: %s"), *Property->GetPathName());
+			}
 		}
-	}
-	OutResult += ValueStr;
-	return bComplete;
-}
-
-void FEmitDefaultValueHelper::InnerGenerate(const UProperty* Property, const uint8* ValuePtr, const FString& PathToMember, FString& OutResult)
-{
-	if (Property->HasAnyPropertyFlags(CPF_NativeAccessSpecifierPrivate | CPF_NativeAccessSpecifierProtected))
-	{
-		UE_LOG(LogK2Compiler, Warning, TEXT("FEmitDefaultValueHelper Cannot access property: %s"), *Property->GetPathName());
-		return;
-	}
+		OutResult += ValueStr;
+		return bComplete;
+	};
 
 	auto StructProperty = Cast<const UStructProperty>(Property);
 	check(!StructProperty || StructProperty->Struct);
@@ -742,16 +757,10 @@ void FEmitDefaultValueHelper::InnerGenerate(const UProperty* Property, const uin
 
 	if (StructProperty)
 	{
+		const FString LocalPathToMember = FString::Printf(TEXT("%s."), *PathToMember);
 		for (auto LocalProperty : TFieldRange<const UProperty>(StructProperty->Struct))
 		{
-			const bool bStaticArray = (Property->ArrayDim > 1);
-			for (int32 ArrayIndex = 0; ArrayIndex < Property->ArrayDim; ++ArrayIndex)
-			{
-				const uint8* LocalValuePtr = LocalProperty->ContainerPtrToValuePtr<uint8>(ValuePtr, ArrayIndex);
-				const FString ArrayPost = bStaticArray ? FString::Printf(TEXT("[%d]"), ArrayIndex) : TEXT("");
-				const FString LocalPathToMember = FString::Printf(TEXT("%s.%s%s"), *PathToMember, *LocalProperty->GetNameCPP(), *ArrayPost);
-				InnerGenerate(LocalProperty, LocalValuePtr, LocalPathToMember, OutResult);
-			}
+			OuterGenerate(LocalProperty, ValuePtr, LocalPathToMember, nullptr, OutResult);
 		}
 	}
 	
@@ -804,9 +813,60 @@ bool FEmitDefaultValueHelper::HandleSpecialTypes(const UProperty* Property, cons
 	return false;
 }
 
-FString FEmitDefaultValueHelper::GenerateConstructor(UClass* BPGC)
+bool FEmitDefaultValueHelper::HandleNonNativeComponent(UBlueprintGeneratedClass* BPGC, FName Name, bool bNew, const FString MemberPath, FString& OutResult)
+{
+	USCS_Node* Node = nullptr;
+	for (UBlueprintGeneratedClass* NodeOwner = BPGC; NodeOwner && !Node; NodeOwner = Cast<UBlueprintGeneratedClass>(NodeOwner->GetSuperClass()))
+	{
+		Node = NodeOwner->SimpleConstructionScript ? NodeOwner->SimpleConstructionScript->FindSCSNode(Name) : nullptr;
+	}
+
+	if (!Node)
+	{
+		return false;
+	}
+
+	ensure(bNew == (Node->GetOuter() == BPGC->SimpleConstructionScript));
+	UObject* ComponentTemplate = bNew 
+		? Node->ComponentTemplate
+		: (BPGC->InheritableComponentHandler ? BPGC->InheritableComponentHandler->GetOverridenComponentTemplate(Node) : nullptr);
+
+	if (!ComponentTemplate)
+	{
+		return false;
+	}
+
+	auto ComponentClass = ComponentTemplate->GetClass();
+	if (bNew)
+	{
+		OutResult += FString::Printf(TEXT("%s = CreateDefaultSubobject<%s%s>(TEXT(\"%s\"));\n")
+			, *MemberPath, ComponentClass->GetPrefixCPP(), *ComponentClass->GetName(), *ComponentTemplate->GetName());
+	}
+
+	UObject* ParentTemplateComponent = bNew 
+		? nullptr 
+		: Node->GetActualComponentTemplate(Cast<UBlueprintGeneratedClass>(BPGC->GetSuperClass()));
+	ensure(bNew || ParentTemplateComponent);
+	ensure(ParentTemplateComponent != ComponentTemplate);
+	ensure(!ParentTemplateComponent || (ParentTemplateComponent->GetClass() == ComponentClass));
+
+	const UObject* ObjectToCompare = bNew
+		? ComponentClass->GetDefaultObject(false)
+		: ParentTemplateComponent;
+	const FString LocalPathToMember = FString::Printf(TEXT("%s->"), *MemberPath);
+	for (auto Property : TFieldRange<const UProperty>(ComponentClass))
+	{
+		OuterGenerate(Property
+			, reinterpret_cast<const uint8*>(ComponentTemplate), LocalPathToMember
+			, reinterpret_cast<const uint8*>(ObjectToCompare), OutResult);
+	}
+
+	return true;
+}
+
+FString FEmitDefaultValueHelper::GenerateConstructor(UClass* InBPGC)
 {
-	check(BPGC);
+	auto BPGC = CastChecked<UBlueprintGeneratedClass>(InBPGC);
 	const FString CppClassName = FString(BPGC->GetPrefixCPP()) + BPGC->GetName();
 
 	FString Result;
@@ -815,24 +875,32 @@ FString FEmitDefaultValueHelper::GenerateConstructor(UClass* BPGC)
 	UObject* CDO = BPGC->GetDefaultObject(false);
 	UObject* ParentCDO = BPGC->GetSuperClass()->GetDefaultObject(false);
 	check(CDO && ParentCDO);
+
+	const bool bIsActor = BPGC->IsChildOf<AActor>();
 	for (auto Property : TFieldRange<const UProperty>(BPGC))
 	{
 		const bool bNewProperty = Property->GetOwnerStruct() == BPGC;
 		const bool bIsAccessible = bNewProperty || !Property->HasAnyPropertyFlags(CPF_NativeAccessSpecifierPrivate);
 		if (bIsAccessible)
 		{
-			const bool bStaticArray = (Property->ArrayDim > 1);
-			for (int32 ArrayIndex = 0; ArrayIndex < Property->ArrayDim; ++ArrayIndex)
+			auto ObjectProperty = Cast<UObjectProperty>(Property);
+			const bool bComponentProp = ObjectProperty && ObjectProperty->PropertyClass && ObjectProperty->PropertyClass->IsChildOf<UActorComponent>();
+			const bool bNullValue = ObjectProperty && (nullptr == ObjectProperty->GetPropertyValue_InContainer(CDO));
+			if (bIsActor && bComponentProp && bNullValue)
 			{
-				const bool bChangedValue = !bNewProperty && !Property->Identical_InContainer(CDO, ParentCDO, ArrayIndex);
-				if (bNewProperty || bChangedValue)
+				if (HandleNonNativeComponent(BPGC, Property->GetFName(), bNewProperty, FString::Printf(TEXT("\t%s"), *Property->GetNameCPP()), Result))
 				{
-					const uint8* ValuePtr = Property->ContainerPtrToValuePtr<uint8>(CDO, ArrayIndex);
-					const FString ArrayPost = bStaticArray ? FString::Printf(TEXT("[%d]"), ArrayIndex) : TEXT("");
-					const FString PathToMember = FString::Printf(TEXT("\t%s"), *Property->GetNameCPP(), *ArrayPost);
-					InnerGenerate(Property, ValuePtr, PathToMember, Result);
+					continue;
 				}
 			}
+
+			//1. Instanced object
+			//2. Component templates outside SCS
+				// FindComponentTemplateByName in UCLass
+				// Random Referenced Objects in UClass (Timelines, other Templates) - Let's fill them while creating CDO ?
+			//3. Native, changed component ???
+
+			OuterGenerate(Property, reinterpret_cast<const uint8*>(CDO), TEXT("\t"), bNewProperty ? nullptr : reinterpret_cast<const uint8*>(ParentCDO), Result, true);
 		}
 	}
 
diff --git a/Engine/Source/Developer/BlueprintCompilerCppBackend/Private/BlueprintCompilerCppBackendUtils.h b/Engine/Source/Developer/BlueprintCompilerCppBackend/Private/BlueprintCompilerCppBackendUtils.h
--- a/Engine/Source/Developer/BlueprintCompilerCppBackend/Private/BlueprintCompilerCppBackendUtils.h
+++ b/Engine/Source/Developer/BlueprintCompilerCppBackend/Private/BlueprintCompilerCppBackendUtils.h
@@ -71,11 +71,11 @@ struct FEmitDefaultValueHelper
 
 	static FString GenerateConstructor(UClass* BPGC);
 private:
+
+	static void OuterGenerate(const UProperty* Property, const uint8* DataContainer, const FString OuterPath, const uint8* OptionalDefaultDataContainer, FString& OutResult, bool bAllowProtected = false);
 	static void InnerGenerate(const UProperty* Property, const uint8* ValuePtr, const FString& PathToMember, FString& OutResult);
 	
-	// Return if it's complete construction
-	static bool OneLineConstruction(const UProperty* Property, const uint8* ValuePtr, FString& OutResult);
-
 	//Return if handled
 	static bool HandleSpecialTypes(const UProperty* Property, const uint8* ValuePtr, FString& OutResult);
+	static bool HandleNonNativeComponent(UBlueprintGeneratedClass* BPGC, FName Name, bool bNew, const FString MemberPath, FString& OutResult);
 };
\ No newline at end of file
diff --git a/Engine/Source/Runtime/Engine/Private/ActorConstruction.cpp b/Engine/Source/Runtime/Engine/Private/ActorConstruction.cpp
--- a/Engine/Source/Runtime/Engine/Private/ActorConstruction.cpp
+++ b/Engine/Source/Runtime/Engine/Private/ActorConstruction.cpp
@@ -506,6 +506,18 @@ void AActor::ExecuteConstruction(const FTransform& Transform, const FComponentIn
 			}
 		}
 	}
+	else
+	{
+#if WITH_EDITOR
+		bool bDoUserConstructionScript;
+		GConfig->GetBool(TEXT("Kismet"), TEXT("bTurnOffEditorConstructionScript"), bDoUserConstructionScript, GEngineIni);
+		if (!GIsEditor || !bDoUserConstructionScript)
+#endif
+		{
+			// Then run the user script, which is responsible for calling its own super, if desired
+			ProcessUserConstructionScript();
+		}
+	}
 
 	GetWorld()->UpdateCullDistanceVolumes(this);
 
