Cis fix

diff --git a/Engine/Source/Developer/BlueprintCompilerCppBackend/Private/BlueprintCompilerCppBackendUtils.cpp b/Engine/Source/Developer/BlueprintCompilerCppBackend/Private/BlueprintCompilerCppBackendUtils.cpp
--- a/Engine/Source/Developer/BlueprintCompilerCppBackend/Private/BlueprintCompilerCppBackendUtils.cpp
+++ b/Engine/Source/Developer/BlueprintCompilerCppBackend/Private/BlueprintCompilerCppBackendUtils.cpp
@@ -716,14 +716,14 @@ FString FEmitDefaultValueHelper::GenerateGetDefaultValue(const UUserDefinedStruc
 
 void FEmitDefaultValueHelper::InnerGenerate(const UProperty* Property, const uint8* ValuePtr, const FString& PathToMember, FString& OutResult)
 {
-	auto OneLineConstruction = [](const UProperty* Property, const uint8* ValuePtr, FString& OutResult) -> bool
+	auto OneLineConstruction = [](const UProperty* Property, const uint8* LocalValuePtr, FString& OutResult) -> bool
 	{
 		FString ValueStr;
 		bool bComplete = true;
-		if (!HandleSpecialTypes(Property, ValuePtr, ValueStr))
+		if (!HandleSpecialTypes(Property, LocalValuePtr, ValueStr))
 		{
 			auto StructProperty = Cast<const UStructProperty>(Property);
-			Property->ExportTextItem(ValueStr, ValuePtr, ValuePtr, nullptr, EPropertyPortFlags::PPF_ExportCpp);
+			Property->ExportTextItem(ValueStr, LocalValuePtr, LocalValuePtr, nullptr, EPropertyPortFlags::PPF_ExportCpp);
 			if (ValueStr.IsEmpty() && StructProperty)
 			{
 				check(StructProperty->Struct);
