Attempted fix for shadowed variable warning

diff --git a/Engine/Source/Developer/BlueprintCompilerCppBackend/Private/BlueprintCompilerCppBackendUtils.cpp b/Engine/Source/Developer/BlueprintCompilerCppBackend/Private/BlueprintCompilerCppBackendUtils.cpp
--- a/Engine/Source/Developer/BlueprintCompilerCppBackend/Private/BlueprintCompilerCppBackendUtils.cpp
+++ b/Engine/Source/Developer/BlueprintCompilerCppBackend/Private/BlueprintCompilerCppBackendUtils.cpp
@@ -716,14 +716,14 @@ FString FEmitDefaultValueHelper::GenerateGetDefaultValue(const UUserDefinedStruc
 
 void FEmitDefaultValueHelper::InnerGenerate(const UProperty* Property, const uint8* ValuePtr, const FString& PathToMember, FString& OutResult)
 {
-	auto OneLineConstruction = [](const UProperty* Property, const uint8* LocalValuePtr, FString& OutResult) -> bool
+	auto OneLineConstruction = [](const UProperty* LocalProperty, const uint8* LocalValuePtr, FString& OutResult) -> bool
 	{
 		FString ValueStr;
 		bool bComplete = true;
-		if (!HandleSpecialTypes(Property, LocalValuePtr, ValueStr))
+		if (!HandleSpecialTypes(LocalProperty, LocalValuePtr, ValueStr))
 		{
-			auto StructProperty = Cast<const UStructProperty>(Property);
-			Property->ExportTextItem(ValueStr, LocalValuePtr, LocalValuePtr, nullptr, EPropertyPortFlags::PPF_ExportCpp);
+			auto StructProperty = Cast<const UStructProperty>(LocalProperty);
+			LocalProperty->ExportTextItem(ValueStr, LocalValuePtr, LocalValuePtr, nullptr, EPropertyPortFlags::PPF_ExportCpp);
 			if (ValueStr.IsEmpty() && StructProperty)
 			{
 				check(StructProperty->Struct);
@@ -732,7 +732,7 @@ void FEmitDefaultValueHelper::InnerGenerate(const UProperty* Property, const uin
 			}
 			if (ValueStr.IsEmpty())
 			{
-				UE_LOG(LogK2Compiler, Warning, TEXT("FEmitDefaultValueHelper Cannot generate initilization: %s"), *Property->GetPathName());
+				UE_LOG(LogK2Compiler, Warning, TEXT("FEmitDefaultValueHelper Cannot generate initilization: %s"), *LocalProperty->GetPathName());
 			}
 		}
 		OutResult += ValueStr;
