Fix another shadowed variable warning

diff --git a/Engine/Source/Developer/BlueprintCompilerCppBackend/Private/BlueprintCompilerCppBackendUtils.cpp b/Engine/Source/Developer/BlueprintCompilerCppBackend/Private/BlueprintCompilerCppBackendUtils.cpp
--- a/Engine/Source/Developer/BlueprintCompilerCppBackend/Private/BlueprintCompilerCppBackendUtils.cpp
+++ b/Engine/Source/Developer/BlueprintCompilerCppBackend/Private/BlueprintCompilerCppBackendUtils.cpp
@@ -716,7 +716,7 @@ FString FEmitDefaultValueHelper::GenerateGetDefaultValue(const UUserDefinedStruc
 
 void FEmitDefaultValueHelper::InnerGenerate(const UProperty* Property, const uint8* ValuePtr, const FString& PathToMember, FString& OutResult)
 {
-	auto OneLineConstruction = [](const UProperty* LocalProperty, const uint8* LocalValuePtr, FString& OutResult) -> bool
+	auto OneLineConstruction = [](const UProperty* LocalProperty, const uint8* LocalValuePtr, FString& LocalOutResult) -> bool
 	{
 		FString ValueStr;
 		bool bComplete = true;
@@ -735,7 +735,7 @@ void FEmitDefaultValueHelper::InnerGenerate(const UProperty* Property, const uin
 				UE_LOG(LogK2Compiler, Warning, TEXT("FEmitDefaultValueHelper Cannot generate initilization: %s"), *LocalProperty->GetPathName());
 			}
 		}
-		OutResult += ValueStr;
+		LocalOutResult += ValueStr;
 		return bComplete;
 	};
 
