UEBP-40 C++ code generation (WIP):
- timelines
- "AddComponent" support
- class subobjects are stored in UClass::MiscObjects
- various fixed, FDefaultValueHelperContext

diff --git a/Engine/Source/Developer/BlueprintCompilerCppBackend/Private/BlueprintCompilerCppBackend.cpp b/Engine/Source/Developer/BlueprintCompilerCppBackend/Private/BlueprintCompilerCppBackend.cpp
--- a/Engine/Source/Developer/BlueprintCompilerCppBackend/Private/BlueprintCompilerCppBackend.cpp
+++ b/Engine/Source/Developer/BlueprintCompilerCppBackend/Private/BlueprintCompilerCppBackend.cpp
@@ -300,6 +300,7 @@ FString FBlueprintCompilerCppBackend::EmitCallStatmentInner(FBlueprintCompiledSt
 		FSafeContextScopedEmmitter SafeContextScope(Result, bUseSafeContext ? Statement.FunctionContext : nullptr, *this, TEXT("\t\t\t"));
 		ensure(!bInline || !SafeContextScope.IsSafeContextUsed());
 
+		bool bCloseCast = false;
 		if (!bInline)
 		{
 			Result += TEXT("\t\t\t");
@@ -309,7 +310,19 @@ FString FBlueprintCompilerCppBackend::EmitCallStatmentInner(FBlueprintCompiledSt
 			UProperty* FuncToCallReturnProperty = Statement.FunctionToCall->GetReturnProperty();
 			if (FuncToCallReturnProperty && ensure(Statement.LHS))
 			{
-				Result += FString::Printf(TEXT("%s = "), *TermToText(Statement.LHS));
+				FString BeginCast;
+				if (auto ObjectProperty = Cast<UObjectProperty>(FuncToCallReturnProperty))
+				{
+					UClass* LClass = Statement.LHS ? Cast<UClass>(Statement.LHS->Type.PinSubCategoryObject.Get()) : nullptr;
+					if (LClass && LClass->IsChildOf(ObjectProperty->PropertyClass) && !ObjectProperty->PropertyClass->IsChildOf(LClass))
+					{
+						BeginCast = FString::Printf(TEXT("CastChecked<%s%s>("), LClass->GetPrefixCPP(), *LClass->GetName());
+						bCloseCast = true;
+					}
+				}
+
+				Result += FString::Printf(TEXT("%s = %s"), *TermToText(Statement.LHS), *BeginCast);
+				
 			}
 		}
 
@@ -353,6 +366,11 @@ FString FBlueprintCompilerCppBackend::EmitCallStatmentInner(FBlueprintCompiledSt
 		Result += EmitMethodInputParameterList(Statement);
 		Result += TEXT(")");
 
+		if (bCloseCast)
+		{
+			Result += TEXT(", ECastCheckedType::NullAllowed)");
+		}
+
 		if (!bInline)
 		{
 			Result += TEXT(";\n");
diff --git a/Engine/Source/Developer/BlueprintCompilerCppBackend/Private/BlueprintCompilerCppBackendUtils.cpp b/Engine/Source/Developer/BlueprintCompilerCppBackend/Private/BlueprintCompilerCppBackendUtils.cpp
--- a/Engine/Source/Developer/BlueprintCompilerCppBackend/Private/BlueprintCompilerCppBackendUtils.cpp
+++ b/Engine/Source/Developer/BlueprintCompilerCppBackend/Private/BlueprintCompilerCppBackendUtils.cpp
@@ -666,20 +666,110 @@ bool FEmitHelper::GenerateAssignmentCast(const FEdGraphPinType& LType, const FEd
 		OutCastEnd = TEXT(")");
 		return true;
 	}
+
 	return false;
 }
 
-void FEmitDefaultValueHelper::OuterGenerate(const UProperty* Property, const uint8* DataContainer, const FString OuterPath, const uint8* OptionalDefaultDataContainer, FString& OutResult, bool bAllowProtected)
+struct FDefaultValueHelperContext
 {
-	if (Property->HasAnyPropertyFlags(CPF_EditorOnly))
+private:
+	FString Indent;
+	FString Result;
+
+	FString AfterCreateFormat;
+	FString ActualClassExpression;
+	UClass* ActualClass;
+
+	TMap<UObject*, FString> NativeObjectNames;
+
+	TArray<UObject*> OrderedObjects;
+public:
+	FDefaultValueHelperContext()
+		: ActualClass(nullptr)
+	{}
+
+	bool FindObject(UObject* Object, FString& OutNamePath, FString* OutAfterCreateFormatPtr = nullptr)
 	{
-		UE_LOG(LogK2Compiler, Log, TEXT("FEmitDefaultValueHelper Skip EditorOnly property: %s"), *Property->GetPathName());
-		return;
+		if (Object == ActualClass)
+		{
+			if (OutAfterCreateFormatPtr)
+			{
+				*OutAfterCreateFormatPtr = AfterCreateFormat;
+			}
+			OutNamePath = ActualClassExpression;
+			return true;
+		}
+
+		if (FString* NamePtr = NativeObjectNames.Find(Object))
+		{
+			OutNamePath = *NamePtr;
+			if (OutAfterCreateFormatPtr)
+			{
+				*OutAfterCreateFormatPtr = FString();
+			}
+			return true;
+		}
+
+		return false;
+	}
+
+	void SetCurrentlyGeneratedClass(UClass* InClass, const FString& AfterCreateFormatWhenTheObjectIsOuter, const FString& NamePath)
+	{
+		ActualClass = InClass;
+		ActualClassExpression = NamePath;
+		AfterCreateFormat = AfterCreateFormatWhenTheObjectIsOuter;
+	}
+
+	FString AddNewObject(UObject* Object)
+	{
+		check(!NativeObjectNames.Contains(Object));
+		const FString UniqueNameBase = TEXT("__Local__");
+		FString UniqueName = UniqueNameBase;
+		for (int32 Index = 0; nullptr != NativeObjectNames.FindKey(UniqueName); ++Index)
+		{
+			UniqueName = FString::Printf(TEXT("%s%d"), *UniqueNameBase, Index);
+		}
+		NativeObjectNames.Add(Object, UniqueName);
+		OrderedObjects.Add(Object);
+		return UniqueName;
+	}
+
+	UClass* GetCurrentlyGeneratedClass() const
+	{
+		return ActualClass;
+	}
+	
+	void IncreaseIndent()
+	{
+		Indent += TEXT("\t");
+	}
+
+	void DecreaseIndent()
+	{
+		Indent.RemoveFromEnd(TEXT("\t"));
+	}
+
+	void AddLine(const FString& Line)
+	{
+		Result += FString::Printf(TEXT("%s%s\n"), *Indent, *Line);
+	}
+
+	FString GetResult() const
+	{
+		return Result;
 	}
 
-	if (Property->HasAnyPropertyFlags(CPF_NativeAccessSpecifierPrivate) || (!bAllowProtected && Property->HasAnyPropertyFlags(CPF_NativeAccessSpecifierProtected)))
+	const TArray<UObject*>& GetOrderedObjects()
+	{
+		return OrderedObjects;
+	}
+};
+
+void FEmitDefaultValueHelper::OuterGenerate(FDefaultValueHelperContext& Context, const UProperty* Property, const uint8* DataContainer, const FString OuterPath, const uint8* OptionalDefaultDataContainer, bool bAllowProtected)
+{
+	if (Property->HasAnyPropertyFlags(CPF_EditorOnly | CPF_Transient))
 	{
-		UE_LOG(LogK2Compiler, Warning, TEXT("FEmitDefaultValueHelper Cannot access property: %s"), *Property->GetPathName());
+		UE_LOG(LogK2Compiler, Log, TEXT("FEmitDefaultValueHelper Skip EditorOnly or Transient property: %s"), *Property->GetPathName());
 		return;
 	}
 
@@ -688,10 +778,16 @@ void FEmitDefaultValueHelper::OuterGenerate(const UProperty* Property, const uin
 	{
 		if (!OptionalDefaultDataContainer || !Property->Identical_InContainer(DataContainer, OptionalDefaultDataContainer, ArrayIndex))
 		{
+			if (Property->HasAnyPropertyFlags(CPF_NativeAccessSpecifierPrivate) || (!bAllowProtected && Property->HasAnyPropertyFlags(CPF_NativeAccessSpecifierProtected)))
+			{
+				UE_LOG(LogK2Compiler, Error, TEXT("FEmitDefaultValueHelper Cannot access property: %s"), *Property->GetPathName());
+				return;
+			}
+
 			const uint8* ValuePtr = Property->ContainerPtrToValuePtr<uint8>(DataContainer, ArrayIndex);
 			const FString ArrayPost = bStaticArray ? FString::Printf(TEXT("[%d]"), ArrayIndex) : TEXT("");
 			const FString PathToMember = FString::Printf(TEXT("%s%s%s"), *OuterPath, *Property->GetNameCPP(), *ArrayPost);
-			InnerGenerate(Property, ValuePtr, PathToMember, OutResult);
+			InnerGenerate(Context, Property, ValuePtr, PathToMember);
 		}
 	}
 }
@@ -705,22 +801,25 @@ FString FEmitDefaultValueHelper::GenerateGetDefaultValue(const UUserDefinedStruc
 	FStructOnScope StructData(Struct);
 	FStructureEditorUtils::Fill_MakeStructureDefaultValue(Struct, StructData.GetStructMemory());
 
+	FDefaultValueHelperContext Context;
+	Context.IncreaseIndent();
+	Context.IncreaseIndent();
 	for (auto Property : TFieldRange<const UProperty>(Struct))
 	{
-		OuterGenerate(Property, StructData.GetStructMemory(), TEXT("\t\tDefaultData__."), nullptr, Result);
+		OuterGenerate(Context, Property, StructData.GetStructMemory(), TEXT("DefaultData__."), nullptr);
 	}
-
+	Result += Context.GetResult();
 	Result += TEXT("\n\t\treturn DefaultData__;\n\t}\n");
 	return Result;
 }
 
-void FEmitDefaultValueHelper::InnerGenerate(const UProperty* Property, const uint8* ValuePtr, const FString& PathToMember, FString& OutResult)
+void FEmitDefaultValueHelper::InnerGenerate(FDefaultValueHelperContext& Context, const UProperty* Property, const uint8* ValuePtr, const FString& PathToMember)
 {
-	auto OneLineConstruction = [](const UProperty* LocalProperty, const uint8* LocalValuePtr, FString& LocalOutResult) -> bool
+	auto OneLineConstruction = [](FDefaultValueHelperContext& Context, const UProperty* LocalProperty, const uint8* LocalValuePtr, FString& OutSingleLine) -> bool
 	{
-		FString ValueStr;
 		bool bComplete = true;
-		if (!HandleSpecialTypes(LocalProperty, LocalValuePtr, ValueStr))
+		FString ValueStr = HandleSpecialTypes(Context, LocalProperty, LocalValuePtr);
+		if (ValueStr.IsEmpty())
 		{
 			auto StructProperty = Cast<const UStructProperty>(LocalProperty);
 			LocalProperty->ExportTextItem(ValueStr, LocalValuePtr, LocalValuePtr, nullptr, EPropertyPortFlags::PPF_ExportCpp);
@@ -732,10 +831,10 @@ void FEmitDefaultValueHelper::InnerGenerate(const UProperty* Property, const uin
 			}
 			if (ValueStr.IsEmpty())
 			{
-				UE_LOG(LogK2Compiler, Warning, TEXT("FEmitDefaultValueHelper Cannot generate initilization: %s"), *LocalProperty->GetPathName());
+				UE_LOG(LogK2Compiler, Error, TEXT("FEmitDefaultValueHelper Cannot generate initilization: %s"), *LocalProperty->GetPathName());
 			}
 		}
-		LocalOutResult += ValueStr;
+		OutSingleLine += ValueStr;
 		return bComplete;
 	};
 
@@ -747,8 +846,8 @@ void FEmitDefaultValueHelper::InnerGenerate(const UProperty* Property, const uin
 	if (!Property->GetOuter()->IsA<UArrayProperty>())
 	{
 		FString ValueStr;
-		const bool bComplete = OneLineConstruction(Property, ValuePtr, ValueStr);
-		OutResult += FString::Printf(TEXT("%s = %s;\n"), *PathToMember, *ValueStr);
+		const bool bComplete = OneLineConstruction(Context, Property, ValuePtr, ValueStr);
+		Context.AddLine(FString::Printf(TEXT("%s = %s;"), *PathToMember, *ValueStr));
 		if (bComplete && !ArrayProperty)
 		{
 			return;
@@ -760,7 +859,7 @@ void FEmitDefaultValueHelper::InnerGenerate(const UProperty* Property, const uin
 		const FString LocalPathToMember = FString::Printf(TEXT("%s."), *PathToMember);
 		for (auto LocalProperty : TFieldRange<const UProperty>(StructProperty->Struct))
 		{
-			OuterGenerate(LocalProperty, ValuePtr, LocalPathToMember, nullptr, OutResult);
+			OuterGenerate(Context, LocalProperty, ValuePtr, LocalPathToMember, nullptr);
 		}
 	}
 	
@@ -772,20 +871,40 @@ void FEmitDefaultValueHelper::InnerGenerate(const UProperty* Property, const uin
 			const uint8* LocalValuePtr = ScriptArrayHelper.GetRawPtr(Index);
 
 			FString ValueStr;
-			const bool bComplete = OneLineConstruction(ArrayProperty->Inner, LocalValuePtr, ValueStr);
-			OutResult += FString::Printf(TEXT("%s.Add(%s);\n"), *PathToMember, *ValueStr);
+			const bool bComplete = OneLineConstruction(Context, ArrayProperty->Inner, LocalValuePtr, ValueStr);
+			Context.AddLine(FString::Printf(TEXT("%s.Add(%s);"), *PathToMember, *ValueStr));
 			if (!bComplete)
 			{
 				const FString LocalPathToMember = FString::Printf(TEXT("%s[%d]"), *PathToMember, Index);
-				InnerGenerate(ArrayProperty->Inner, LocalValuePtr, LocalPathToMember, OutResult);
+				InnerGenerate(Context, ArrayProperty->Inner, LocalValuePtr, LocalPathToMember);
 			}
 		}
 	}
 }
 
-bool FEmitDefaultValueHelper::HandleSpecialTypes(const UProperty* Property, const uint8* ValuePtr, FString& OutResult)
+FString FEmitDefaultValueHelper::HandleSpecialTypes(FDefaultValueHelperContext& Context, const UProperty* Property, const uint8* ValuePtr)
 {
 	//TODO: Use Path maps for Objects
+	if (auto ObjectProperty = Cast<UObjectProperty>(Property))
+	{
+		UObject* Object = ObjectProperty->GetPropertyValue(ValuePtr);
+		if (Object)
+		{
+			FString NativeName;
+			if (Context.FindObject(Object, NativeName))
+			{
+				return NativeName;
+			}
+
+			auto BPGC = Context.GetCurrentlyGeneratedClass();
+			if (BPGC && Object->GetOuter()->GetName() == BPGC->GetName()) // hack to fix curve deuplication
+			{
+				// needs to check if in native code, it will be still subobject of a class
+
+				return HandleClassSubobject(Context, Object);
+			}
+		}
+	}
 
 	if (auto StructProperty = Cast<UStructProperty>(Property))
 	{
@@ -796,51 +915,56 @@ bool FEmitDefaultValueHelper::HandleSpecialTypes(const UProperty* Property, cons
 			const auto Rotation = Transform->GetRotation();
 			const auto Translation = Transform->GetTranslation();
 			const auto Scale = Transform->GetScale3D();
-			OutResult = FString::Printf(TEXT("FTransform(FQuat(%f, %f, %f, %f), FVector(%f, %f, %f), FVector(%f, %f, %f))")
+			return FString::Printf(TEXT("FTransform(FQuat(%f, %f, %f, %f), FVector(%f, %f, %f), FVector(%f, %f, %f))")
 				, Rotation.X, Rotation.Y, Rotation.Z, Rotation.W
 				, Translation.X, Translation.Y, Translation.Z
 				, Scale.X, Scale.Y, Scale.Z);
-			return true;
 		}
 
 		if (TBaseStructure<FVector>::Get() == StructProperty->Struct)
 		{
 			const FVector* Vector = reinterpret_cast<const FVector*>(ValuePtr);
-			OutResult = FString::Printf(TEXT("FVector(%f, %f, %f)"), Vector->X, Vector->Y, Vector->Z);
-			return true;
+			return FString::Printf(TEXT("FVector(%f, %f, %f)"), Vector->X, Vector->Y, Vector->Z);
+		}
+
+		if (TBaseStructure<FGuid>::Get() == StructProperty->Struct)
+		{
+			const FGuid* Guid = reinterpret_cast<const FGuid*>(ValuePtr);
+			return FString::Printf(TEXT("FGuid(0x%08X, 0x%08X, 0x%08X, 0x%08X)"), Guid->A, Guid->B, Guid->C, Guid->D);
 		}
 	}
-	return false;
+	return FString();
 }
 
-bool FEmitDefaultValueHelper::HandleNonNativeComponent(UBlueprintGeneratedClass* BPGC, FName Name, bool bNew, const FString MemberPath, FString& OutResult)
+UActorComponent* FEmitDefaultValueHelper::HandleNonNativeComponent(FDefaultValueHelperContext& Context, UBlueprintGeneratedClass* BPGC, FName ObjectName, bool bNew, const FString& NativeName)
 {
 	USCS_Node* Node = nullptr;
 	for (UBlueprintGeneratedClass* NodeOwner = BPGC; NodeOwner && !Node; NodeOwner = Cast<UBlueprintGeneratedClass>(NodeOwner->GetSuperClass()))
 	{
-		Node = NodeOwner->SimpleConstructionScript ? NodeOwner->SimpleConstructionScript->FindSCSNode(Name) : nullptr;
+		Node = NodeOwner->SimpleConstructionScript ? NodeOwner->SimpleConstructionScript->FindSCSNode(ObjectName) : nullptr;
 	}
 
 	if (!Node)
 	{
-		return false;
+		return nullptr;
 	}
 
 	ensure(bNew == (Node->GetOuter() == BPGC->SimpleConstructionScript));
-	UObject* ComponentTemplate = bNew 
+	UActorComponent* ComponentTemplate = bNew
 		? Node->ComponentTemplate
 		: (BPGC->InheritableComponentHandler ? BPGC->InheritableComponentHandler->GetOverridenComponentTemplate(Node) : nullptr);
 
 	if (!ComponentTemplate)
 	{
-		return false;
+		return nullptr;
 	}
 
 	auto ComponentClass = ComponentTemplate->GetClass();
 	if (bNew)
 	{
-		OutResult += FString::Printf(TEXT("%s = CreateDefaultSubobject<%s%s>(TEXT(\"%s\"));\n")
-			, *MemberPath, ComponentClass->GetPrefixCPP(), *ComponentClass->GetName(), *ComponentTemplate->GetName());
+		Context.AddLine(TEXT(""));
+		Context.AddLine(FString::Printf(TEXT("%s = CreateDefaultSubobject<%s%s>(TEXT(\"%s\"));")
+			, *NativeName, ComponentClass->GetPrefixCPP(), *ComponentClass->GetName(), *ComponentTemplate->GetName()));
 	}
 
 	UObject* ParentTemplateComponent = bNew 
@@ -853,15 +977,15 @@ bool FEmitDefaultValueHelper::HandleNonNativeComponent(UBlueprintGeneratedClass*
 	const UObject* ObjectToCompare = bNew
 		? ComponentClass->GetDefaultObject(false)
 		: ParentTemplateComponent;
-	const FString LocalPathToMember = FString::Printf(TEXT("%s->"), *MemberPath);
+	const FString LocalPathToMember = FString::Printf(TEXT("%s->"), *NativeName);
 	for (auto Property : TFieldRange<const UProperty>(ComponentClass))
 	{
-		OuterGenerate(Property
+		OuterGenerate(Context, Property
 			, reinterpret_cast<const uint8*>(ComponentTemplate), LocalPathToMember
-			, reinterpret_cast<const uint8*>(ObjectToCompare), OutResult);
+			, reinterpret_cast<const uint8*>(ObjectToCompare));
 	}
 
-	return true;
+	return ComponentTemplate;
 }
 
 FString FEmitDefaultValueHelper::GenerateConstructor(UClass* InBPGC)
@@ -869,41 +993,159 @@ FString FEmitDefaultValueHelper::GenerateConstructor(UClass* InBPGC)
 	auto BPGC = CastChecked<UBlueprintGeneratedClass>(InBPGC);
 	const FString CppClassName = FString(BPGC->GetPrefixCPP()) + BPGC->GetName();
 
-	FString Result;
-	Result += FString::Printf(TEXT("%s::%s(const FObjectInitializer& ObjectInitializer) : Super(ObjectInitializer)\n{\n"), *CppClassName, *CppClassName);
+	FDefaultValueHelperContext Context;
+	Context.AddLine(FString::Printf(TEXT("%s::%s(const FObjectInitializer& ObjectInitializer) : Super(ObjectInitializer)"), *CppClassName, *CppClassName));
+	Context.AddLine(TEXT("{"));
+	Context.IncreaseIndent();
+
+	// When CDO is created create all subobjects owned by the class
+	TArray<UActorComponent*> ActorComponentTempatesOwnedByClass = BPGC->ComponentTemplates;
+	{
+		// Gather all CT from SCS and IH, the remaining ones are generated for class..
+		if (auto SCS = BPGC->SimpleConstructionScript)
+		{
+			for (auto Node : SCS->GetAllNodes())
+			{
+				ActorComponentTempatesOwnedByClass.RemoveSwap(Node->ComponentTemplate);
+			}
+		}
+		if (auto IH = BPGC->GetInheritableComponentHandler())
+		{
+			TArray<UActorComponent*> AllTemplates;
+			IH->GetAllTemplates(AllTemplates);
+			ActorComponentTempatesOwnedByClass.RemoveAllSwap([&](UActorComponent* Component) -> bool
+			{
+				return AllTemplates.Contains(Component);
+			});
+		}
+
+		Context.AddLine(FString::Printf(TEXT("if(HasAnyFlags(RF_ClassDefaultObject) && (%s::StaticClass() == GetClass()))"), *CppClassName));
+		Context.AddLine(TEXT("{"));
+		Context.IncreaseIndent();
+		Context.AddLine(TEXT("ensure(0 == GetClass()->MiscObjects.Num());"));
+
+		Context.SetCurrentlyGeneratedClass(BPGC, TEXT("GetClass()->MiscObjects.Add(%s);"), TEXT("GetClass()"));
+
+		for (auto ComponentTemplate : ActorComponentTempatesOwnedByClass)
+		{
+			if (ComponentTemplate)
+			{
+				HandleClassSubobject(Context, ComponentTemplate);
+			}
+		}
+
+		for (auto TimelineTemplate : BPGC->Timelines)
+		{
+			if (TimelineTemplate)
+			{
+				HandleClassSubobject(Context, TimelineTemplate);
+			}
+		}
+		Context.DecreaseIndent();
+		Context.AddLine(TEXT("}"));
+
+		// Let's have an easy access to generated class subobjects
+		Context.AddLine(TEXT("{")); // no shadow variables
+		Context.IncreaseIndent();
+
+		int32 SubobjectIndex = 0;
+		for (UObject* Obj : Context.GetOrderedObjects())
+		{
+			FString NativeName;
+			const bool bFound = Context.FindObject(Obj, NativeName);
+			ensure(bFound);
+			Context.AddLine(FString::Printf(TEXT("auto %s = CastChecked<%s%s>(GetClass()->MiscObjects[%d]);")
+				, *NativeName
+				, Obj->GetClass()->GetPrefixCPP()
+				, *Obj->GetClass()->GetName()
+				, SubobjectIndex));
+			SubobjectIndex++;
+		}
+	}
 
 	UObject* CDO = BPGC->GetDefaultObject(false);
 	UObject* ParentCDO = BPGC->GetSuperClass()->GetDefaultObject(false);
 	check(CDO && ParentCDO);
+	Context.AddLine(TEXT(""));
 
+	// TIMELINES
+	TSet<const UProperty*> HandledProperties;
 	const bool bIsActor = BPGC->IsChildOf<AActor>();
+	if (bIsActor)
+	{
+		for (auto TimelineTemplate : BPGC->Timelines)
+		{
+			FString NativeName;
+			if (Context.FindObject(TimelineTemplate, NativeName))
+			{
+				Context.AddLine(FString::Printf(TEXT("UBlueprintGeneratedClass::CreateTimelineComponent(this, %s);"), *NativeName));
+				auto Prop = FindField<UObjectPropertyBase>(BPGC, *UTimelineTemplate::TimelineTemplateNameToVariableName(TimelineTemplate->GetFName()));
+				if (Prop)
+				{
+					HandledProperties.Add(Prop);
+				}
+			}
+		}
+	}
+	
+	// COMPONENTS
 	for (auto Property : TFieldRange<const UProperty>(BPGC))
 	{
 		const bool bNewProperty = Property->GetOwnerStruct() == BPGC;
 		const bool bIsAccessible = bNewProperty || !Property->HasAnyPropertyFlags(CPF_NativeAccessSpecifierPrivate);
-		if (bIsAccessible)
+		if (bIsAccessible && !HandledProperties.Contains(Property))
 		{
 			auto ObjectProperty = Cast<UObjectProperty>(Property);
 			const bool bComponentProp = ObjectProperty && ObjectProperty->PropertyClass && ObjectProperty->PropertyClass->IsChildOf<UActorComponent>();
 			const bool bNullValue = ObjectProperty && (nullptr == ObjectProperty->GetPropertyValue_InContainer(CDO));
 			if (bIsActor && bComponentProp && bNullValue)
 			{
-				if (HandleNonNativeComponent(BPGC, Property->GetFName(), bNewProperty, FString::Printf(TEXT("\t%s"), *Property->GetNameCPP()), Result))
+				auto HandledComponent = HandleNonNativeComponent(Context, BPGC, Property->GetFName(), bNewProperty, Property->GetNameCPP());
+				if (HandledComponent)
 				{
+					ensure(!ActorComponentTempatesOwnedByClass.Contains(HandledComponent));
 					continue;
 				}
 			}
 
-			//1. Instanced object
-			//2. Component templates outside SCS
-				// FindComponentTemplateByName in UCLass
-				// Random Referenced Objects in UClass (Timelines, other Templates) - Let's fill them while creating CDO ?
-			//3. Native, changed component ???
-
-			OuterGenerate(Property, reinterpret_cast<const uint8*>(CDO), TEXT("\t"), bNewProperty ? nullptr : reinterpret_cast<const uint8*>(ParentCDO), Result, true);
+			OuterGenerate(Context, Property, reinterpret_cast<const uint8*>(CDO), TEXT(""), bNewProperty ? nullptr : reinterpret_cast<const uint8*>(ParentCDO), true);
 		}
 	}
 
-	Result += TEXT("\n}\n");
-	return Result;
+	Context.DecreaseIndent();
+	Context.AddLine(TEXT("}"));
+	Context.DecreaseIndent();
+	Context.AddLine(TEXT("}"));
+
+	return Context.GetResult();
+}
+
+FString FEmitDefaultValueHelper::HandleClassSubobject(FDefaultValueHelperContext& Context, UObject* Object)
+{
+	FString OuterStr;
+	FString AfterCreateFormat;
+	if (!Context.FindObject(Object->GetOuter(), OuterStr, &AfterCreateFormat))
+	{
+		return FString();
+	}
+
+	const FString LocalNativeName = Context.AddNewObject(Object);
+	UClass* ObjectClass = Object->GetClass();
+	Context.AddLine(FString::Printf(TEXT("auto %s = NewObject<%s%s>(%s, TEXT(\"%s\"));")
+		, *LocalNativeName, ObjectClass->GetPrefixCPP(), *ObjectClass->GetName()
+		, *OuterStr, *Object->GetName()));
+	if (!AfterCreateFormat.IsEmpty())
+	{
+		Context.AddLine(FString::Printf(*AfterCreateFormat, *LocalNativeName));
+	}
+
+	const FString MemberPath = FString::Printf(TEXT("%s->"), *LocalNativeName);
+	for (auto Property : TFieldRange<const UProperty>(ObjectClass))
+	{
+		OuterGenerate(Context, Property
+			, reinterpret_cast<const uint8*>(Object), MemberPath
+			, reinterpret_cast<const uint8*>(ObjectClass->GetDefaultObject(false)));
+	}
+
+	return LocalNativeName;
 }
\ No newline at end of file
diff --git a/Engine/Source/Developer/BlueprintCompilerCppBackend/Private/BlueprintCompilerCppBackendUtils.h b/Engine/Source/Developer/BlueprintCompilerCppBackend/Private/BlueprintCompilerCppBackendUtils.h
--- a/Engine/Source/Developer/BlueprintCompilerCppBackend/Private/BlueprintCompilerCppBackendUtils.h
+++ b/Engine/Source/Developer/BlueprintCompilerCppBackend/Private/BlueprintCompilerCppBackendUtils.h
@@ -65,6 +65,8 @@ struct FEmitHelper
 	static bool GenerateAssignmentCast(const FEdGraphPinType& LType, const FEdGraphPinType& RType, FString& OutCastBegin, FString& OutCastEnd);
 };
 
+struct FDefaultValueHelperContext;
+
 struct FEmitDefaultValueHelper
 {
 	static FString GenerateGetDefaultValue(const UUserDefinedStruct* Struct);
@@ -72,10 +74,19 @@ struct FEmitDefaultValueHelper
 	static FString GenerateConstructor(UClass* BPGC);
 private:
 
-	static void OuterGenerate(const UProperty* Property, const uint8* DataContainer, const FString OuterPath, const uint8* OptionalDefaultDataContainer, FString& OutResult, bool bAllowProtected = false);
-	static void InnerGenerate(const UProperty* Property, const uint8* ValuePtr, const FString& PathToMember, FString& OutResult);
+	// OuterPath ends with "->" or "."
+	static void OuterGenerate(FDefaultValueHelperContext& Context, const UProperty* Property, const uint8* DataContainer, const FString OuterPath, const uint8* OptionalDefaultDataContainer, bool bAllowProtected = false);
+	
+	// PathToMember ends with variable name
+	static void InnerGenerate(FDefaultValueHelperContext& Context, const UProperty* Property, const uint8* ValuePtr, const FString& PathToMember);
+	
+	// Returns native term, 
+	// returns empty string if cannot handle
+	static FString HandleSpecialTypes(FDefaultValueHelperContext& Context, const UProperty* Property, const uint8* ValuePtr);
+
+	static UActorComponent* HandleNonNativeComponent(FDefaultValueHelperContext& Context, UBlueprintGeneratedClass* BPGC, FName ObjectName, bool bNew, const FString& NativeName);
 	
-	//Return if handled
-	static bool HandleSpecialTypes(const UProperty* Property, const uint8* ValuePtr, FString& OutResult);
-	static bool HandleNonNativeComponent(UBlueprintGeneratedClass* BPGC, FName Name, bool bNew, const FString MemberPath, FString& OutResult);
+	// Creates the subobject (of class) returns it's native local name, 
+	// returns empty string if cannot handle
+	static FString HandleClassSubobject(FDefaultValueHelperContext& Context, UObject* Object);
 };
\ No newline at end of file
diff --git a/Engine/Source/Editor/UnrealEd/Private/Kismet2/Kismet2.cpp b/Engine/Source/Editor/UnrealEd/Private/Kismet2/Kismet2.cpp
--- a/Engine/Source/Editor/UnrealEd/Private/Kismet2/Kismet2.cpp
+++ b/Engine/Source/Editor/UnrealEd/Private/Kismet2/Kismet2.cpp
@@ -963,6 +963,28 @@ void FKismetEditorUtilities::GenerateCppCode(UObject* Obj, TSharedPtr<FString> O
 		//TGuardValue<bool> DuplicatingReadOnly(InBlueprintObj->bDuplicatingReadOnly, true);
 		{
 			auto BlueprintObj = DuplicateObject<UBlueprint>(InBlueprintObj, GetTransientPackage(), *InBlueprintObj->GetName());
+			{
+				// FIX NOT-DUPLICATED CURVES
+
+				auto NewTimelineOwnerClass = CastChecked<UBlueprintGeneratedClass>(BlueprintObj->GeneratedClass);
+				TSet<UCurveBase*> AllCurves;
+				for (auto Timeline : NewTimelineOwnerClass->Timelines)
+				{
+					Timeline->GetAllCurves(AllCurves);
+				}
+				TMap<UCurveBase*, UCurveBase*> CurveMap;
+				ensure(InBlueprintObj->GetOutermost() != NewTimelineOwnerClass->GetOutermost());
+				for (auto OldCurve : AllCurves)
+				{
+					if (OldCurve && OldCurve->GetOutermost() == InBlueprintObj->GetOutermost())
+					{
+						UCurveBase* NewCurve = DuplicateObject<UCurveBase>(OldCurve, NewTimelineOwnerClass);
+						CurveMap.Add(OldCurve, NewCurve);
+					}
+				}
+				FArchiveReplaceObjectRef<UCurveBase> ReplaceCurvesAr(NewTimelineOwnerClass, CurveMap, /*bNullPrivateRefs=*/ false, /*bIgnoreOuterRef=*/ true, /*bIgnoreArchetypeRef=*/ true);
+			}
+
 			{
 				auto Reinstancer = FBlueprintCompileReinstancer::Create(BlueprintObj->GeneratedClass);
 
diff --git a/Engine/Source/Runtime/CoreUObject/Private/UObject/Class.cpp b/Engine/Source/Runtime/CoreUObject/Private/UObject/Class.cpp
--- a/Engine/Source/Runtime/CoreUObject/Private/UObject/Class.cpp
+++ b/Engine/Source/Runtime/CoreUObject/Private/UObject/Class.cpp
@@ -2525,6 +2525,8 @@ void UClass::AddReferencedObjects(UObject* InThis, FReferenceCollector& Collecto
 		This->CallAddReferencedObjects(This->ClassDefaultObject, Collector);
 	}
 
+	Collector.AddReferencedObjects(This->MiscObjects, This);
+
 	Super::AddReferencedObjects( This, Collector );
 }
 
@@ -4441,6 +4443,13 @@ UScriptStruct* TBaseStructure<FRandomStream>::Get()
 	return ScriptStruct;
 }
 
+UScriptStruct* TBaseStructure<FGuid>::Get()
+{
+	static auto ScriptStruct = StaticGetBaseStructureInternal(TEXT("Guid"));
+	return ScriptStruct;
+}
+
+
 UScriptStruct* TBaseStructure<FFallbackStruct>::Get()
 {
 	static auto ScriptStruct = StaticGetBaseStructureInternal(TEXT("FallbackStruct"));
diff --git a/Engine/Source/Runtime/CoreUObject/Public/UObject/Class.h b/Engine/Source/Runtime/CoreUObject/Public/UObject/Class.h
--- a/Engine/Source/Runtime/CoreUObject/Public/UObject/Class.h
+++ b/Engine/Source/Runtime/CoreUObject/Public/UObject/Class.h
@@ -1934,6 +1934,9 @@ public:
 	/** This class's native functions. */
 	TArray<FNativeFunctionLookup> NativeFunctionLookupTable;
 
+	/** Misc Objects Owned by the class. Mostly templates from BP-generated native class */
+	TArray<UObject*> MiscObjects;
+
 public:
 	// Constructors
 	UClass(const FObjectInitializer& ObjectInitializer = FObjectInitializer::Get());
@@ -2790,6 +2793,11 @@ template<> struct TBaseStructure<FRandomStream>
 	COREUOBJECT_API static UScriptStruct* Get();
 };
 
+template<> struct TBaseStructure<FGuid>
+{
+	COREUOBJECT_API static UScriptStruct* Get();
+};
+
 struct FFallbackStruct;
 template<> struct TBaseStructure<FFallbackStruct>
 {
diff --git a/Engine/Source/Runtime/Engine/Classes/Curves/CurveBase.h b/Engine/Source/Runtime/Engine/Classes/Curves/CurveBase.h
--- a/Engine/Source/Runtime/Engine/Classes/Curves/CurveBase.h
+++ b/Engine/Source/Runtime/Engine/Classes/Curves/CurveBase.h
@@ -401,7 +401,6 @@ public:
 	UPROPERTY()
 	TEnumAsByte<ERichCurveExtrapolation> PostInfinityExtrap;
 
-private:
 	/** Sorted array of keys */
 	UPROPERTY()
 	TArray<FRichCurveKey> Keys;
diff --git a/Engine/Source/Runtime/Engine/Classes/Engine/BlueprintGeneratedClass.h b/Engine/Source/Runtime/Engine/Classes/Engine/BlueprintGeneratedClass.h
--- a/Engine/Source/Runtime/Engine/Classes/Engine/BlueprintGeneratedClass.h
+++ b/Engine/Source/Runtime/Engine/Classes/Engine/BlueprintGeneratedClass.h
@@ -452,6 +452,7 @@ public:
 
 	/** Create Timeline objects for this Actor based on the Timelines array*/
 	virtual void CreateComponentsForActor(AActor* Actor) const;
+	static void CreateTimelineComponent(AActor* Actor, const UTimelineTemplate* TimelineTemplate);
 
 	// UObject interface
 	virtual void Serialize(FArchive& Ar) override;
diff --git a/Engine/Source/Runtime/Engine/Private/ActorConstruction.cpp b/Engine/Source/Runtime/Engine/Private/ActorConstruction.cpp
--- a/Engine/Source/Runtime/Engine/Private/ActorConstruction.cpp
+++ b/Engine/Source/Runtime/Engine/Private/ActorConstruction.cpp
@@ -609,15 +609,22 @@ UActorComponent* AActor::CreateComponentFromTemplate(UActorComponent* Template,
 UActorComponent* AActor::AddComponent(FName TemplateName, bool bManualAttachment, const FTransform& RelativeTransform, const UObject* ComponentTemplateContext)
 {
 	UActorComponent* Template = nullptr;
-	UBlueprintGeneratedClass* BlueprintGeneratedClass = Cast<UBlueprintGeneratedClass>((ComponentTemplateContext != nullptr) ? ComponentTemplateContext->GetClass() : GetClass());
-	while(BlueprintGeneratedClass != nullptr)
+	for (UClass* TemplateOwnerClass = (ComponentTemplateContext != nullptr) ? ComponentTemplateContext->GetClass() : GetClass()
+		; TemplateOwnerClass && !Template
+		; TemplateOwnerClass = TemplateOwnerClass->GetSuperClass())
 	{
-		Template = BlueprintGeneratedClass->FindComponentTemplateByName(TemplateName);
-		if(nullptr != Template)
+		if (auto BPGC = Cast<UBlueprintGeneratedClass>(TemplateOwnerClass))
 		{
-			break;
+			Template = BPGC->FindComponentTemplateByName(TemplateName);
+		}
+		else
+		{
+			UObject** FoundTemplatePtr = TemplateOwnerClass->MiscObjects.FindByPredicate([=](UObject* Obj) -> bool
+			{
+				return Obj && Obj->IsA<UActorComponent>() && (Obj->GetFName() == TemplateName);
+			});
+			Template = (nullptr != FoundTemplatePtr) ? Cast<UActorComponent>(*FoundTemplatePtr) : nullptr;
 		}
-		BlueprintGeneratedClass = Cast<UBlueprintGeneratedClass>(BlueprintGeneratedClass->GetSuperClass());
 	}
 
 	bool bIsSceneComponent = false;
diff --git a/Engine/Source/Runtime/Engine/Private/BlueprintGeneratedClass.cpp b/Engine/Source/Runtime/Engine/Private/BlueprintGeneratedClass.cpp
--- a/Engine/Source/Runtime/Engine/Private/BlueprintGeneratedClass.cpp
+++ b/Engine/Source/Runtime/Engine/Private/BlueprintGeneratedClass.cpp
@@ -407,122 +407,134 @@ UActorComponent* UBlueprintGeneratedClass::FindComponentTemplateByName(const FNa
 	return NULL;
 }
 
-void UBlueprintGeneratedClass::CreateComponentsForActor(AActor* Actor) const
+void UBlueprintGeneratedClass::CreateTimelineComponent(AActor* Actor, const UTimelineTemplate* TimelineTemplate)
 {
-	check(Actor != NULL);
-	check(!Actor->IsTemplate());
-	check(!Actor->IsPendingKill());
-
-	// Iterate over each timeline template
-	for(int32 i=0; i<Timelines.Num(); i++)
+	if (!Actor
+		|| !TimelineTemplate
+		|| !TimelineTemplate->bValidatedAsWired
+		|| Actor->IsTemplate()
+		|| Actor->IsPendingKill())
 	{
-		const UTimelineTemplate* TimelineTemplate = Timelines[i];
+		return;
+	}
 
-		// Not fatal if NULL, but shouldn't happen and ignored if not wired up in graph
-		if(!TimelineTemplate||!TimelineTemplate->bValidatedAsWired)
-		{
-			continue;
-		}
+	FName NewName(*UTimelineTemplate::TimelineTemplateNameToVariableName(TimelineTemplate->GetFName()));
+	UTimelineComponent* NewTimeline = NewObject<UTimelineComponent>(Actor, NewName);
+	NewTimeline->CreationMethod = EComponentCreationMethod::UserConstructionScript; // Indicate it comes from a blueprint so it gets cleared when we rerun construction scripts
+	Actor->BlueprintCreatedComponents.Add(NewTimeline); // Add to array so it gets saved
+	NewTimeline->SetNetAddressable();	// This component has a stable name that can be referenced for replication
 
-		FName NewName( *UTimelineTemplate::TimelineTemplateNameToVariableName( TimelineTemplate->GetFName() ));
-		UTimelineComponent* NewTimeline = NewObject<UTimelineComponent>(Actor, NewName);
-		NewTimeline->CreationMethod = EComponentCreationMethod::UserConstructionScript; // Indicate it comes from a blueprint so it gets cleared when we rerun construction scripts
-		Actor->BlueprintCreatedComponents.Add(NewTimeline); // Add to array so it gets saved
-		NewTimeline->SetNetAddressable();	// This component has a stable name that can be referenced for replication
+	NewTimeline->SetPropertySetObject(Actor); // Set which object the timeline should drive properties on
+	NewTimeline->SetDirectionPropertyName(TimelineTemplate->GetDirectionPropertyName());
 
-		NewTimeline->SetPropertySetObject(Actor); // Set which object the timeline should drive properties on
-		NewTimeline->SetDirectionPropertyName(TimelineTemplate->GetDirectionPropertyName());
+	NewTimeline->SetTimelineLength(TimelineTemplate->TimelineLength); // copy length
+	NewTimeline->SetTimelineLengthMode(TimelineTemplate->LengthMode);
 
-		NewTimeline->SetTimelineLength(TimelineTemplate->TimelineLength); // copy length
-		NewTimeline->SetTimelineLengthMode(TimelineTemplate->LengthMode);
+	// Find property with the same name as the template and assign the new Timeline to it
+	UClass* ActorClass = Actor->GetClass();
+	UObjectPropertyBase* Prop = FindField<UObjectPropertyBase>(ActorClass, *UTimelineTemplate::TimelineTemplateNameToVariableName(TimelineTemplate->GetFName()));
+	if (Prop)
+	{
+		Prop->SetObjectPropertyValue_InContainer(Actor, NewTimeline);
+	}
 
-		// Find property with the same name as the template and assign the new Timeline to it
-		UClass* ActorClass = Actor->GetClass();
-		UObjectPropertyBase* Prop = FindField<UObjectPropertyBase>( ActorClass, *UTimelineTemplate::TimelineTemplateNameToVariableName(TimelineTemplate->GetFName()) );
-		if(Prop)
+	// Event tracks
+	// In the template there is a track for each function, but in the runtime Timeline each key has its own delegate, so we fold them together
+	for (int32 TrackIdx = 0; TrackIdx < TimelineTemplate->EventTracks.Num(); TrackIdx++)
+	{
+		const FTTEventTrack* EventTrackTemplate = &TimelineTemplate->EventTracks[TrackIdx];
+		if (EventTrackTemplate->CurveKeys != NULL)
 		{
-			Prop->SetObjectPropertyValue_InContainer(Actor, NewTimeline);
-		}
+			// Create delegate for all keys in this track
+			FScriptDelegate EventDelegate;
+			EventDelegate.BindUFunction(Actor, TimelineTemplate->GetEventTrackFunctionName(TrackIdx));
 
-		// Event tracks
-		// In the template there is a track for each function, but in the runtime Timeline each key has its own delegate, so we fold them together
-		for(int32 TrackIdx=0; TrackIdx<TimelineTemplate->EventTracks.Num(); TrackIdx++)
-		{
-			const FTTEventTrack* EventTrackTemplate = &TimelineTemplate->EventTracks[TrackIdx];
-			if(EventTrackTemplate->CurveKeys != NULL)
+			// Create an entry in Events for each key of this track
+			for (auto It(EventTrackTemplate->CurveKeys->FloatCurve.GetKeyIterator()); It; ++It)
 			{
-				// Create delegate for all keys in this track
-				FScriptDelegate EventDelegate;
-				EventDelegate.BindUFunction(Actor, TimelineTemplate->GetEventTrackFunctionName(TrackIdx));
-
-				// Create an entry in Events for each key of this track
-				for (auto It(EventTrackTemplate->CurveKeys->FloatCurve.GetKeyIterator()); It; ++It)
-				{
-					NewTimeline->AddEvent(It->Time, FOnTimelineEvent(EventDelegate));
-				}
+				NewTimeline->AddEvent(It->Time, FOnTimelineEvent(EventDelegate));
 			}
 		}
+	}
 
-		// Float tracks
-		for(int32 TrackIdx=0; TrackIdx<TimelineTemplate->FloatTracks.Num(); TrackIdx++)
+	// Float tracks
+	for (int32 TrackIdx = 0; TrackIdx < TimelineTemplate->FloatTracks.Num(); TrackIdx++)
+	{
+		const FTTFloatTrack* FloatTrackTemplate = &TimelineTemplate->FloatTracks[TrackIdx];
+		if (FloatTrackTemplate->CurveFloat != NULL)
 		{
-			const FTTFloatTrack* FloatTrackTemplate = &TimelineTemplate->FloatTracks[TrackIdx];
-			if(FloatTrackTemplate->CurveFloat != NULL)
-			{
-				NewTimeline->AddInterpFloat(FloatTrackTemplate->CurveFloat, FOnTimelineFloat(), TimelineTemplate->GetTrackPropertyName(FloatTrackTemplate->TrackName));
-			}
+			NewTimeline->AddInterpFloat(FloatTrackTemplate->CurveFloat, FOnTimelineFloat(), TimelineTemplate->GetTrackPropertyName(FloatTrackTemplate->TrackName));
 		}
+	}
 
-		// Vector tracks
-		for(int32 TrackIdx=0; TrackIdx<TimelineTemplate->VectorTracks.Num(); TrackIdx++)
+	// Vector tracks
+	for (int32 TrackIdx = 0; TrackIdx < TimelineTemplate->VectorTracks.Num(); TrackIdx++)
+	{
+		const FTTVectorTrack* VectorTrackTemplate = &TimelineTemplate->VectorTracks[TrackIdx];
+		if (VectorTrackTemplate->CurveVector != NULL)
 		{
-			const FTTVectorTrack* VectorTrackTemplate = &TimelineTemplate->VectorTracks[TrackIdx];
-			if(VectorTrackTemplate->CurveVector != NULL)
-			{
-				NewTimeline->AddInterpVector(VectorTrackTemplate->CurveVector, FOnTimelineVector(), TimelineTemplate->GetTrackPropertyName(VectorTrackTemplate->TrackName));
-			}
+			NewTimeline->AddInterpVector(VectorTrackTemplate->CurveVector, FOnTimelineVector(), TimelineTemplate->GetTrackPropertyName(VectorTrackTemplate->TrackName));
 		}
+	}
 
-		// Linear color tracks
-		for(int32 TrackIdx=0; TrackIdx<TimelineTemplate->LinearColorTracks.Num(); TrackIdx++)
+	// Linear color tracks
+	for (int32 TrackIdx = 0; TrackIdx < TimelineTemplate->LinearColorTracks.Num(); TrackIdx++)
+	{
+		const FTTLinearColorTrack* LinearColorTrackTemplate = &TimelineTemplate->LinearColorTracks[TrackIdx];
+		if (LinearColorTrackTemplate->CurveLinearColor != NULL)
 		{
-			const FTTLinearColorTrack* LinearColorTrackTemplate = &TimelineTemplate->LinearColorTracks[TrackIdx];
-			if(LinearColorTrackTemplate->CurveLinearColor != NULL)
-			{
-				NewTimeline->AddInterpLinearColor(LinearColorTrackTemplate->CurveLinearColor, FOnTimelineLinearColor(), TimelineTemplate->GetTrackPropertyName(LinearColorTrackTemplate->TrackName));
-			}
+			NewTimeline->AddInterpLinearColor(LinearColorTrackTemplate->CurveLinearColor, FOnTimelineLinearColor(), TimelineTemplate->GetTrackPropertyName(LinearColorTrackTemplate->TrackName));
 		}
+	}
 
-		// Set up delegate that gets called after all properties are updated
-		FScriptDelegate UpdateDelegate;
-		UpdateDelegate.BindUFunction(Actor, TimelineTemplate->GetUpdateFunctionName());
-		NewTimeline->SetTimelinePostUpdateFunc(FOnTimelineEvent(UpdateDelegate));
+	// Set up delegate that gets called after all properties are updated
+	FScriptDelegate UpdateDelegate;
+	UpdateDelegate.BindUFunction(Actor, TimelineTemplate->GetUpdateFunctionName());
+	NewTimeline->SetTimelinePostUpdateFunc(FOnTimelineEvent(UpdateDelegate));
 
-		// Set up finished delegate that gets called after all properties are updated
-		FScriptDelegate FinishedDelegate;
-		FinishedDelegate.BindUFunction(Actor, TimelineTemplate->GetFinishedFunctionName());
-		NewTimeline->SetTimelineFinishedFunc(FOnTimelineEvent(FinishedDelegate));
+	// Set up finished delegate that gets called after all properties are updated
+	FScriptDelegate FinishedDelegate;
+	FinishedDelegate.BindUFunction(Actor, TimelineTemplate->GetFinishedFunctionName());
+	NewTimeline->SetTimelineFinishedFunc(FOnTimelineEvent(FinishedDelegate));
 
-		NewTimeline->RegisterComponent();
+	NewTimeline->RegisterComponent();
 
-		// Start playing now, if desired
-		if(TimelineTemplate->bAutoPlay)
-		{
-			// Needed for autoplay timelines in cooked builds, since they won't have Activate() called via the Play call below
-			NewTimeline->bAutoActivate = true;
-			NewTimeline->Play();
-		}
+	// Start playing now, if desired
+	if (TimelineTemplate->bAutoPlay)
+	{
+		// Needed for autoplay timelines in cooked builds, since they won't have Activate() called via the Play call below
+		NewTimeline->bAutoActivate = true;
+		NewTimeline->Play();
+	}
 
-		// Set to loop, if desired
-		if(TimelineTemplate->bLoop)
-		{
-			NewTimeline->SetLooping(true);
-		}
+	// Set to loop, if desired
+	if (TimelineTemplate->bLoop)
+	{
+		NewTimeline->SetLooping(true);
+	}
 
-		// Set replication, if desired
-		if(TimelineTemplate->bReplicated)
+	// Set replication, if desired
+	if (TimelineTemplate->bReplicated)
+	{
+		NewTimeline->SetIsReplicated(true);
+	}
+}
+
+void UBlueprintGeneratedClass::CreateComponentsForActor(AActor* Actor) const
+{
+	check(Actor != NULL);
+	check(!Actor->IsTemplate());
+	check(!Actor->IsPendingKill());
+
+	// Iterate over each timeline template
+	for(int32 i=0; i<Timelines.Num(); i++)
+	{
+		const UTimelineTemplate* TimelineTemplate = Timelines[i];
+
+		// Not fatal if NULL, but shouldn't happen and ignored if not wired up in graph
+		if(TimelineTemplate && TimelineTemplate->bValidatedAsWired)
 		{
-			NewTimeline->SetIsReplicated(true);
+			CreateTimelineComponent(Actor, TimelineTemplate);
 		}
 	}
 }
