Shadowed variable build warning fix.

diff --git a/Engine/Source/Developer/BlueprintCompilerCppBackend/Private/BlueprintCompilerCppBackendUtils.cpp b/Engine/Source/Developer/BlueprintCompilerCppBackend/Private/BlueprintCompilerCppBackendUtils.cpp
--- a/Engine/Source/Developer/BlueprintCompilerCppBackend/Private/BlueprintCompilerCppBackendUtils.cpp
+++ b/Engine/Source/Developer/BlueprintCompilerCppBackend/Private/BlueprintCompilerCppBackendUtils.cpp
@@ -815,10 +815,10 @@ FString FEmitDefaultValueHelper::GenerateGetDefaultValue(const UUserDefinedStruc
 
 void FEmitDefaultValueHelper::InnerGenerate(FDefaultValueHelperContext& Context, const UProperty* Property, const uint8* ValuePtr, const FString& PathToMember)
 {
-	auto OneLineConstruction = [](FDefaultValueHelperContext& Context, const UProperty* LocalProperty, const uint8* LocalValuePtr, FString& OutSingleLine) -> bool
+	auto OneLineConstruction = [](FDefaultValueHelperContext& LocalContext, const UProperty* LocalProperty, const uint8* LocalValuePtr, FString& OutSingleLine) -> bool
 	{
 		bool bComplete = true;
-		FString ValueStr = HandleSpecialTypes(Context, LocalProperty, LocalValuePtr);
+		FString ValueStr = HandleSpecialTypes(LocalContext, LocalProperty, LocalValuePtr);
 		if (ValueStr.IsEmpty())
 		{
 			auto StructProperty = Cast<const UStructProperty>(LocalProperty);
