UEBP-40 C++ code generation (WIP):
- private properties access
- instanced objects
- object path map (wip)
- All subobjects are created, before any of them is filled

diff --git a/Engine/Source/Developer/BlueprintCompilerCppBackend/Private/BlueprintCompilerCppBackendUtils.cpp b/Engine/Source/Developer/BlueprintCompilerCppBackend/Private/BlueprintCompilerCppBackendUtils.cpp
--- a/Engine/Source/Developer/BlueprintCompilerCppBackend/Private/BlueprintCompilerCppBackendUtils.cpp
+++ b/Engine/Source/Developer/BlueprintCompilerCppBackend/Private/BlueprintCompilerCppBackendUtils.cpp
@@ -685,12 +685,23 @@ private:
 	TMap<UObject*, FString> NativeObjectNames;
 
 	TArray<UObject*> OrderedObjects;
+
+	int32 LocalNameIndexMax;
+
+	TArray<FString> HighPriorityLines;
+	TArray<FString> LowPriorityLines;
 public:
+
+	bool bCreatingObjectsPerClass;
+
 	FDefaultValueHelperContext()
 		: ActualClass(nullptr)
+		, LocalNameIndexMax(0)
+		, bCreatingObjectsPerClass(false)
 	{}
 
-	bool FindObject(UObject* Object, FString& OutNamePath, FString* OutAfterCreateFormatPtr = nullptr)
+// OBJECTS:
+	bool FindLocalObject(UObject* Object, FString& OutNamePath, FString* OutAfterCreateFormatPtr = nullptr) const
 	{
 		if (Object == ActualClass)
 		{
@@ -702,7 +713,7 @@ public:
 			return true;
 		}
 
-		if (FString* NamePtr = NativeObjectNames.Find(Object))
+		if (const FString* NamePtr = NativeObjectNames.Find(Object))
 		{
 			OutNamePath = *NamePtr;
 			if (OutAfterCreateFormatPtr)
@@ -715,6 +726,17 @@ public:
 		return false;
 	}
 
+	FString GenerateGetProperty(const UProperty* Property) const
+	{
+		check(Property);
+		const UStruct* OwnerStruct = Property->GetOwnerStruct();
+		const FString OwnerName = FString(OwnerStruct->GetPrefixCPP()) + OwnerStruct->GetName();
+		const FString OwnerPath = OwnerName + (OwnerStruct->IsA<UClass>() ? TEXT("::StaticClass()") : TEXT("::StaticStruct()"));
+		//return FString::Printf(TEXT("%s->FindPropertyByName(GET_MEMBER_NAME_CHECKED(%s, %s))"), *OwnerPath, *OwnerName, *Property->GetNameCPP());
+
+		return FString::Printf(TEXT("%s->FindPropertyByName(FName(TEXT(\"%s\")))"), *OwnerPath, *Property->GetNameCPP());
+	}
+
 	void SetCurrentlyGeneratedClass(UClass* InClass, const FString& AfterCreateFormatWhenTheObjectIsOuter, const FString& NamePath)
 	{
 		ActualClass = InClass;
@@ -722,15 +744,23 @@ public:
 		AfterCreateFormat = AfterCreateFormatWhenTheObjectIsOuter;
 	}
 
+	FString GenerateUniqueLocalName()
+	{
+		const FString UniqueNameBase = TEXT("__Local__");
+		const FString UniqueName = FString::Printf(TEXT("%s%d"), *UniqueNameBase, LocalNameIndexMax);
+		++LocalNameIndexMax;
+		return UniqueName;
+	}
+
+	void AddObjectFromLocalProperty(UObject* Object, const FString& NativeName)
+	{
+		NativeObjectNames.Add(Object, NativeName);
+	}
+
 	FString AddNewObject(UObject* Object)
 	{
 		check(!NativeObjectNames.Contains(Object));
-		const FString UniqueNameBase = TEXT("__Local__");
-		FString UniqueName = UniqueNameBase;
-		for (int32 Index = 0; nullptr != NativeObjectNames.FindKey(UniqueName); ++Index)
-		{
-			UniqueName = FString::Printf(TEXT("%s%d"), *UniqueNameBase, Index);
-		}
+		const FString UniqueName = GenerateUniqueLocalName();
 		NativeObjectNames.Add(Object, UniqueName);
 		OrderedObjects.Add(Object);
 		return UniqueName;
@@ -740,7 +770,38 @@ public:
 	{
 		return ActualClass;
 	}
-	
+
+	const TArray<UObject*>& GetOrderedObjects()
+	{
+		return OrderedObjects;
+	}
+
+	FString FindGlobalObject(UObject* Object)
+	{
+		// TODO: check if not excluded
+
+		if (auto BPGC = ExactCast<UBlueprintGeneratedClass>(Object))
+		{
+			// TODO: check if supported 
+			return FString::Printf(TEXT("%s%s::StaticClass()"), BPGC->GetPrefixCPP(), *BPGC->GetName());
+		}
+
+		if (auto UDS = Cast<UUserDefinedStruct>(Object))
+		{
+			return FString::Printf(TEXT("%s%s::StaticStruct()"), UDS->GetPrefixCPP(), *UDS->GetName());
+		}
+
+		if (auto UDE = Cast<UUserDefinedEnum>(Object))
+		{
+			return FString::Printf(TEXT("%s_StaticEnum()"), *UDE->GetName());
+		}
+
+		// TODO: handle subobjects
+
+		return FString{};
+	}
+
+// TEXT
 	void IncreaseIndent()
 	{
 		Indent += TEXT("\t");
@@ -751,23 +812,45 @@ public:
 		Indent.RemoveFromEnd(TEXT("\t"));
 	}
 
-	void AddLine(const FString& Line)
+	void AddHighPriorityLine(const FString& Line)
 	{
-		Result += FString::Printf(TEXT("%s%s\n"), *Indent, *Line);
+		HighPriorityLines.Emplace(FString::Printf(TEXT("%s%s\n"), *Indent, *Line));
 	}
 
-	FString GetResult() const
+	void AddLowPriorityLine(const FString& Line)
 	{
-		return Result;
+		LowPriorityLines.Emplace(FString::Printf(TEXT("%s%s\n"), *Indent, *Line));
 	}
 
-	const TArray<UObject*>& GetOrderedObjects()
+	void FlushLines()
 	{
-		return OrderedObjects;
+		for (auto& Line : HighPriorityLines)
+		{
+			Result += Line;
+		}
+		HighPriorityLines.Reset();
+
+		for (auto& Line : LowPriorityLines)
+		{
+			Result += Line;
+		}
+		LowPriorityLines.Reset();
+	}
+
+	FString GetResult()
+	{
+		FlushLines();
+		return Result;
 	}
 };
 
-void FEmitDefaultValueHelper::OuterGenerate(FDefaultValueHelperContext& Context, const UProperty* Property, const uint8* DataContainer, const FString OuterPath, const uint8* OptionalDefaultDataContainer, bool bAllowProtected)
+void FEmitDefaultValueHelper::OuterGenerate(FDefaultValueHelperContext& Context
+	, const UProperty* Property
+	, const FString& OuterPath
+	, const uint8* DataContainer
+	, const uint8* OptionalDefaultDataContainer
+	, EPropertyAccessOperator AccessOperator
+	, bool bAllowProtected)
 {
 	if (Property->HasAnyPropertyFlags(CPF_EditorOnly | CPF_Transient))
 	{
@@ -775,21 +858,39 @@ void FEmitDefaultValueHelper::OuterGenerate(FDefaultValueHelperContext& Context,
 		return;
 	}
 
-	const bool bStaticArray = (Property->ArrayDim > 1);
 	for (int32 ArrayIndex = 0; ArrayIndex < Property->ArrayDim; ++ArrayIndex)
 	{
 		if (!OptionalDefaultDataContainer || !Property->Identical_InContainer(DataContainer, OptionalDefaultDataContainer, ArrayIndex))
 		{
+			FString PathToMember;
 			if (Property->HasAnyPropertyFlags(CPF_NativeAccessSpecifierPrivate) || (!bAllowProtected && Property->HasAnyPropertyFlags(CPF_NativeAccessSpecifierProtected)))
 			{
-				UE_LOG(LogK2Compiler, Error, TEXT("FEmitDefaultValueHelper Cannot access property: %s"), *Property->GetPathName());
-				return;
+				ensure(EPropertyAccessOperator::None != AccessOperator);
+				
+				const FString PropertyObject = Context.GenerateGetProperty(Property); //X::StaticClass()->FindPropertyByName(GET_MEMBER_NAME_CHECKED(X, X))
+
+				FStringOutputDevice TypeDeclaration;
+				const uint32 CppTemplateTypeFlags = EPropertyExportCPPFlags::CPPF_CustomTypeName
+					| EPropertyExportCPPFlags::CPPF_NoConst | EPropertyExportCPPFlags::CPPF_NoRef
+					| EPropertyExportCPPFlags::CPPF_BlueprintCppBackend;
+				Property->ExportCppDeclaration(TypeDeclaration, EExportedDeclaration::Parameter, nullptr, CppTemplateTypeFlags, true);
+
+				const FString OperatorStr = (EPropertyAccessOperator::Dot == AccessOperator) ? TEXT("&") : TEXT("");
+				const FString ContainerStr = (EPropertyAccessOperator::None == AccessOperator) ? TEXT("this") : FString::Printf(TEXT("%s(%s)"), *OperatorStr, *OuterPath);
+				PathToMember = Context.GenerateUniqueLocalName();
+				Context.AddLowPriorityLine(FString::Printf(TEXT("auto& %s = *(%s->ContainerPtrToValuePtr<%s>(%s, %d));")
+					, *PathToMember, *PropertyObject, *TypeDeclaration, *ContainerStr, ArrayIndex));
+			}
+			else
+			{
+				const FString AccessOperatorStr = (EPropertyAccessOperator::None == AccessOperator) ? TEXT("")
+					: ((EPropertyAccessOperator::Pointer == AccessOperator) ? TEXT("->") : TEXT("."));
+				const bool bStaticArray = (Property->ArrayDim > 1);
+				const FString ArrayPost = bStaticArray ? FString::Printf(TEXT("[%d]"), ArrayIndex) : TEXT("");
+				PathToMember = FString::Printf(TEXT("%s%s%s%s"), *OuterPath, *AccessOperatorStr, *Property->GetNameCPP(), *ArrayPost);
 			}
-
 			const uint8* ValuePtr = Property->ContainerPtrToValuePtr<uint8>(DataContainer, ArrayIndex);
-			const FString ArrayPost = bStaticArray ? FString::Printf(TEXT("[%d]"), ArrayIndex) : TEXT("");
-			const FString PathToMember = FString::Printf(TEXT("%s%s%s"), *OuterPath, *Property->GetNameCPP(), *ArrayPost);
-			InnerGenerate(Context, Property, ValuePtr, PathToMember);
+			InnerGenerate(Context, Property, PathToMember, ValuePtr);
 		}
 	}
 }
@@ -808,14 +909,14 @@ FString FEmitDefaultValueHelper::GenerateGetDefaultValue(const UUserDefinedStruc
 	Context.IncreaseIndent();
 	for (auto Property : TFieldRange<const UProperty>(Struct))
 	{
-		OuterGenerate(Context, Property, StructData.GetStructMemory(), TEXT("DefaultData__."), nullptr);
+		OuterGenerate(Context, Property, TEXT("DefaultData__"), StructData.GetStructMemory(), nullptr, EPropertyAccessOperator::Dot);
 	}
 	Result += Context.GetResult();
 	Result += TEXT("\n\t\treturn DefaultData__;\n\t}\n");
 	return Result;
 }
 
-void FEmitDefaultValueHelper::InnerGenerate(FDefaultValueHelperContext& Context, const UProperty* Property, const uint8* ValuePtr, const FString& PathToMember)
+void FEmitDefaultValueHelper::InnerGenerate(FDefaultValueHelperContext& Context, const UProperty* Property, const FString& PathToMember, const uint8* ValuePtr, bool bWithoutFirstConstructionLine)
 {
 	auto OneLineConstruction = [](FDefaultValueHelperContext& LocalContext, const UProperty* LocalProperty, const uint8* LocalValuePtr, FString& OutSingleLine) -> bool
 	{
@@ -845,11 +946,11 @@ void FEmitDefaultValueHelper::InnerGenerate(FDefaultValueHelperContext& Context,
 	auto ArrayProperty = Cast<const UArrayProperty>(Property);
 	check(!ArrayProperty || ArrayProperty->Inner);
 
-	if (!Property->GetOuter()->IsA<UArrayProperty>())
+	if (!bWithoutFirstConstructionLine)
 	{
 		FString ValueStr;
 		const bool bComplete = OneLineConstruction(Context, Property, ValuePtr, ValueStr);
-		Context.AddLine(FString::Printf(TEXT("%s = %s;"), *PathToMember, *ValueStr));
+		Context.AddLowPriorityLine(FString::Printf(TEXT("%s = %s;"), *PathToMember, *ValueStr));
 		if (bComplete && !ArrayProperty)
 		{
 			return;
@@ -858,10 +959,9 @@ void FEmitDefaultValueHelper::InnerGenerate(FDefaultValueHelperContext& Context,
 
 	if (StructProperty)
 	{
-		const FString LocalPathToMember = FString::Printf(TEXT("%s."), *PathToMember);
 		for (auto LocalProperty : TFieldRange<const UProperty>(StructProperty->Struct))
 		{
-			OuterGenerate(Context, LocalProperty, ValuePtr, LocalPathToMember, nullptr);
+			OuterGenerate(Context, LocalProperty, PathToMember, ValuePtr, nullptr, EPropertyAccessOperator::Dot);
 		}
 	}
 	
@@ -874,11 +974,11 @@ void FEmitDefaultValueHelper::InnerGenerate(FDefaultValueHelperContext& Context,
 
 			FString ValueStr;
 			const bool bComplete = OneLineConstruction(Context, ArrayProperty->Inner, LocalValuePtr, ValueStr);
-			Context.AddLine(FString::Printf(TEXT("%s.Add(%s);"), *PathToMember, *ValueStr));
+			Context.AddLowPriorityLine(FString::Printf(TEXT("%s.Add(%s);"), *PathToMember, *ValueStr));
 			if (!bComplete)
 			{
 				const FString LocalPathToMember = FString::Printf(TEXT("%s[%d]"), *PathToMember, Index);
-				InnerGenerate(Context, ArrayProperty->Inner, LocalValuePtr, LocalPathToMember);
+				InnerGenerate(Context, ArrayProperty->Inner, LocalPathToMember, LocalValuePtr, true);
 			}
 		}
 	}
@@ -893,18 +993,32 @@ FString FEmitDefaultValueHelper::HandleSpecialTypes(FDefaultValueHelperContext&
 		if (Object)
 		{
 			FString NativeName;
-			if (Context.FindObject(Object, NativeName))
+			if (Context.FindLocalObject(Object, NativeName))
 			{
 				return NativeName;
 			}
 
 			auto BPGC = Context.GetCurrentlyGeneratedClass();
-			if (BPGC && Object->GetOuter()->GetName() == BPGC->GetName()) // hack to fix curve deuplication
+			auto CDO = BPGC ? BPGC->GetDefaultObject(false) : nullptr;
+			if (BPGC && Object && CDO && Object->IsIn(BPGC) && !Object->IsIn(CDO) && Context.bCreatingObjectsPerClass)
 			{
-				// needs to check if in native code, it will be still subobject of a class
-
 				return HandleClassSubobject(Context, Object);
 			}
+
+			if (!Context.bCreatingObjectsPerClass && Property->HasAnyPropertyFlags(CPF_InstancedReference))
+			{
+				const FString CreateAsInstancedSubobject = HandleInstancedSubobject(Context, Object);
+				if (!CreateAsInstancedSubobject.IsEmpty())
+				{
+					return CreateAsInstancedSubobject;
+				}
+			}
+
+			const FString MappedObject = Context.FindGlobalObject(Object);
+			if (!MappedObject.IsEmpty())
+			{
+				return MappedObject;
+			}
 		}
 	}
 
@@ -964,9 +1078,10 @@ UActorComponent* FEmitDefaultValueHelper::HandleNonNativeComponent(FDefaultValue
 	auto ComponentClass = ComponentTemplate->GetClass();
 	if (bNew)
 	{
-		Context.AddLine(TEXT(""));
-		Context.AddLine(FString::Printf(TEXT("%s = CreateDefaultSubobject<%s%s>(TEXT(\"%s\"));")
+		Context.AddLowPriorityLine(TEXT(""));
+		Context.AddHighPriorityLine(FString::Printf(TEXT("%s = CreateDefaultSubobject<%s%s>(TEXT(\"%s\"));")
 			, *NativeName, ComponentClass->GetPrefixCPP(), *ComponentClass->GetName(), *ComponentTemplate->GetName()));
+		Context.AddObjectFromLocalProperty(ComponentTemplate, NativeName);
 	}
 
 	UObject* ParentTemplateComponent = bNew 
@@ -979,12 +1094,12 @@ UActorComponent* FEmitDefaultValueHelper::HandleNonNativeComponent(FDefaultValue
 	const UObject* ObjectToCompare = bNew
 		? ComponentClass->GetDefaultObject(false)
 		: ParentTemplateComponent;
-	const FString LocalPathToMember = FString::Printf(TEXT("%s->"), *NativeName);
 	for (auto Property : TFieldRange<const UProperty>(ComponentClass))
 	{
-		OuterGenerate(Context, Property
-			, reinterpret_cast<const uint8*>(ComponentTemplate), LocalPathToMember
-			, reinterpret_cast<const uint8*>(ObjectToCompare));
+		OuterGenerate(Context, Property, NativeName
+			, reinterpret_cast<const uint8*>(ComponentTemplate)
+			, reinterpret_cast<const uint8*>(ObjectToCompare)
+			, EPropertyAccessOperator::Pointer);
 	}
 
 	return ComponentTemplate;
@@ -996,8 +1111,8 @@ FString FEmitDefaultValueHelper::GenerateConstructor(UClass* InBPGC)
 	const FString CppClassName = FString(BPGC->GetPrefixCPP()) + BPGC->GetName();
 
 	FDefaultValueHelperContext Context;
-	Context.AddLine(FString::Printf(TEXT("%s::%s(const FObjectInitializer& ObjectInitializer) : Super(ObjectInitializer)"), *CppClassName, *CppClassName));
-	Context.AddLine(TEXT("{"));
+	Context.AddLowPriorityLine(FString::Printf(TEXT("%s::%s(const FObjectInitializer& ObjectInitializer) : Super(ObjectInitializer)"), *CppClassName, *CppClassName));
+	Context.AddLowPriorityLine(TEXT("{"));
 	Context.IncreaseIndent();
 
 	// When CDO is created create all subobjects owned by the class
@@ -1021,13 +1136,13 @@ FString FEmitDefaultValueHelper::GenerateConstructor(UClass* InBPGC)
 			});
 		}
 
-		Context.AddLine(FString::Printf(TEXT("if(HasAnyFlags(RF_ClassDefaultObject) && (%s::StaticClass() == GetClass()))"), *CppClassName));
-		Context.AddLine(TEXT("{"));
+		Context.AddLowPriorityLine(FString::Printf(TEXT("if(HasAnyFlags(RF_ClassDefaultObject) && (%s::StaticClass() == GetClass()))"), *CppClassName));
+		Context.AddLowPriorityLine(TEXT("{"));
 		Context.IncreaseIndent();
-		Context.AddLine(TEXT("ensure(0 == GetClass()->MiscObjects.Num());"));
-
+		Context.AddLowPriorityLine(TEXT("ensure(0 == GetClass()->MiscObjects.Num());"));
+		Context.bCreatingObjectsPerClass = true;
 		Context.SetCurrentlyGeneratedClass(BPGC, TEXT("GetClass()->MiscObjects.Add(%s);"), TEXT("GetClass()"));
-
+		Context.FlushLines();
 		for (auto ComponentTemplate : ActorComponentTempatesOwnedByClass)
 		{
 			if (ComponentTemplate)
@@ -1052,20 +1167,21 @@ FString FEmitDefaultValueHelper::GenerateConstructor(UClass* InBPGC)
 			}
 		}
 
+		Context.bCreatingObjectsPerClass = false;
 		Context.DecreaseIndent();
-		Context.AddLine(TEXT("}"));
+		Context.AddLowPriorityLine(TEXT("}"));
 
 		// Let's have an easy access to generated class subobjects
-		Context.AddLine(TEXT("{")); // no shadow variables
+		Context.AddLowPriorityLine(TEXT("{")); // no shadow variables
 		Context.IncreaseIndent();
-
+		Context.FlushLines();
 		int32 SubobjectIndex = 0;
 		for (UObject* Obj : Context.GetOrderedObjects())
 		{
 			FString NativeName;
-			const bool bFound = Context.FindObject(Obj, NativeName);
+			const bool bFound = Context.FindLocalObject(Obj, NativeName);
 			ensure(bFound);
-			Context.AddLine(FString::Printf(TEXT("auto %s = CastChecked<%s%s>(GetClass()->MiscObjects[%d]);")
+			Context.AddHighPriorityLine(FString::Printf(TEXT("auto %s = CastChecked<%s%s>(GetClass()->MiscObjects[%d]);")
 				, *NativeName
 				, Obj->GetClass()->GetPrefixCPP()
 				, *Obj->GetClass()->GetName()
@@ -1075,9 +1191,11 @@ FString FEmitDefaultValueHelper::GenerateConstructor(UClass* InBPGC)
 	}
 
 	UObject* CDO = BPGC->GetDefaultObject(false);
+	Context.AddObjectFromLocalProperty(CDO, TEXT("this"));
+
 	UObject* ParentCDO = BPGC->GetSuperClass()->GetDefaultObject(false);
 	check(CDO && ParentCDO);
-	Context.AddLine(TEXT(""));
+	Context.AddLowPriorityLine(TEXT(""));
 
 	// TIMELINES
 	TSet<const UProperty*> HandledProperties;
@@ -1093,6 +1211,7 @@ FString FEmitDefaultValueHelper::GenerateConstructor(UClass* InBPGC)
 			auto ObjectProperty = Cast<UObjectProperty>(Property);
 			const bool bComponentProp = ObjectProperty && ObjectProperty->PropertyClass && ObjectProperty->PropertyClass->IsChildOf<UActorComponent>();
 			const bool bNullValue = ObjectProperty && (nullptr == ObjectProperty->GetPropertyValue_InContainer(CDO));
+			const bool bActorComponent = bIsActor && bComponentProp && bNullValue;
 			if (bIsActor && bComponentProp && bNullValue)
 			{
 				auto HandledComponent = HandleNonNativeComponent(Context, BPGC, Property->GetFName(), bNewProperty, Property->GetNameCPP());
@@ -1103,14 +1222,14 @@ FString FEmitDefaultValueHelper::GenerateConstructor(UClass* InBPGC)
 				}
 			}
 
-			OuterGenerate(Context, Property, reinterpret_cast<const uint8*>(CDO), TEXT(""), bNewProperty ? nullptr : reinterpret_cast<const uint8*>(ParentCDO), true);
+			OuterGenerate(Context, Property, TEXT(""), reinterpret_cast<const uint8*>(CDO), bNewProperty ? nullptr : reinterpret_cast<const uint8*>(ParentCDO), EPropertyAccessOperator::None, true);
 		}
 	}
 
 	Context.DecreaseIndent();
-	Context.AddLine(TEXT("}"));
+	Context.AddLowPriorityLine(TEXT("}"));
 	Context.DecreaseIndent();
-	Context.AddLine(TEXT("}"));
+	Context.AddLowPriorityLine(TEXT("}"));
 
 	return Context.GetResult();
 }
@@ -1119,27 +1238,64 @@ FString FEmitDefaultValueHelper::HandleClassSubobject(FDefaultValueHelperContext
 {
 	FString OuterStr;
 	FString AfterCreateFormat;
-	if (!Context.FindObject(Object->GetOuter(), OuterStr, &AfterCreateFormat))
+	if (!Context.FindLocalObject(Object->GetOuter(), OuterStr, &AfterCreateFormat))
 	{
 		return FString();
 	}
 
 	const FString LocalNativeName = Context.AddNewObject(Object);
 	UClass* ObjectClass = Object->GetClass();
-	Context.AddLine(FString::Printf(TEXT("auto %s = NewObject<%s%s>(%s, TEXT(\"%s\"));")
+	Context.AddHighPriorityLine(FString::Printf(TEXT("auto %s = NewObject<%s%s>(%s, TEXT(\"%s\"));")
 		, *LocalNativeName, ObjectClass->GetPrefixCPP(), *ObjectClass->GetName()
 		, *OuterStr, *Object->GetName()));
 	if (!AfterCreateFormat.IsEmpty())
 	{
-		Context.AddLine(FString::Printf(*AfterCreateFormat, *LocalNativeName));
+		Context.AddLowPriorityLine(FString::Printf(*AfterCreateFormat, *LocalNativeName));
+	}
+
+	for (auto Property : TFieldRange<const UProperty>(ObjectClass))
+	{
+		OuterGenerate(Context, Property, LocalNativeName
+			, reinterpret_cast<const uint8*>(Object)
+			, reinterpret_cast<const uint8*>(ObjectClass->GetDefaultObject(false))
+			, EPropertyAccessOperator::Pointer);
+	}
+
+	return LocalNativeName;
+}
+
+FString FEmitDefaultValueHelper::HandleInstancedSubobject(FDefaultValueHelperContext& Context, UObject* Object)
+{
+	check(Object);
+
+	const FString LocalNativeName = Context.AddNewObject(Object);
+	UClass* ObjectClass = Object->GetClass();
+
+	auto BPGC = Context.GetCurrentlyGeneratedClass();
+	auto CDO = BPGC ? BPGC->GetDefaultObject(false) : nullptr;
+	if (ensure(CDO) && (CDO == Object->GetOuter()))
+	{
+		Context.AddHighPriorityLine(FString::Printf(TEXT("auto %s = CreateDefaultSubobject<%s%s>(TEXT(\"%s\"));")
+			, *LocalNativeName, ObjectClass->GetPrefixCPP(), *ObjectClass->GetName(), *Object->GetName()));
+	}
+	else
+	{
+		FString OuterStr;
+		if (!Context.FindLocalObject(Object->GetOuter(), OuterStr))
+		{
+			return FString();
+		}
+		Context.AddHighPriorityLine(FString::Printf(TEXT("auto %s = NewObject<%s%s>(%s, TEXT(\"%s\"));")
+			, *LocalNativeName, ObjectClass->GetPrefixCPP(), *ObjectClass->GetName()
+			, *OuterStr, *Object->GetName()));
 	}
 
-	const FString MemberPath = FString::Printf(TEXT("%s->"), *LocalNativeName);
 	for (auto Property : TFieldRange<const UProperty>(ObjectClass))
 	{
-		OuterGenerate(Context, Property
-			, reinterpret_cast<const uint8*>(Object), MemberPath
-			, reinterpret_cast<const uint8*>(ObjectClass->GetDefaultObject(false)));
+		OuterGenerate(Context, Property, LocalNativeName
+			, reinterpret_cast<const uint8*>(Object)
+			, reinterpret_cast<const uint8*>(ObjectClass->GetDefaultObject(false))
+			, EPropertyAccessOperator::Pointer);
 	}
 
 	return LocalNativeName;
diff --git a/Engine/Source/Developer/BlueprintCompilerCppBackend/Private/BlueprintCompilerCppBackendUtils.h b/Engine/Source/Developer/BlueprintCompilerCppBackend/Private/BlueprintCompilerCppBackendUtils.h
--- a/Engine/Source/Developer/BlueprintCompilerCppBackend/Private/BlueprintCompilerCppBackendUtils.h
+++ b/Engine/Source/Developer/BlueprintCompilerCppBackend/Private/BlueprintCompilerCppBackendUtils.h
@@ -74,11 +74,18 @@ struct FEmitDefaultValueHelper
 	static FString GenerateConstructor(UClass* BPGC);
 private:
 
-	// OuterPath ends with "->" or "."
-	static void OuterGenerate(FDefaultValueHelperContext& Context, const UProperty* Property, const uint8* DataContainer, const FString OuterPath, const uint8* OptionalDefaultDataContainer, bool bAllowProtected = false);
+	enum class EPropertyAccessOperator
+	{
+		None, // for self scope, this
+		Pointer,
+		Dot,
+	};
+
+	// OuterPath ends context/outer name (or empty, if the scope is "this")
+	static void OuterGenerate(FDefaultValueHelperContext& Context, const UProperty* Property, const FString& OuterPath, const uint8* DataContainer, const uint8* OptionalDefaultDataContainer, EPropertyAccessOperator AccessOperator, bool bAllowProtected = false);
 	
 	// PathToMember ends with variable name
-	static void InnerGenerate(FDefaultValueHelperContext& Context, const UProperty* Property, const uint8* ValuePtr, const FString& PathToMember);
+	static void InnerGenerate(FDefaultValueHelperContext& Context, const UProperty* Property, const FString& PathToMember, const uint8* ValuePtr, bool bWithoutFirstConstructionLine = false);
 	
 	// Returns native term, 
 	// returns empty string if cannot handle
@@ -89,4 +96,6 @@ private:
 	// Creates the subobject (of class) returns it's native local name, 
 	// returns empty string if cannot handle
 	static FString HandleClassSubobject(FDefaultValueHelperContext& Context, UObject* Object);
+
+	static FString HandleInstancedSubobject(FDefaultValueHelperContext& Context, UObject* Object);
 };
\ No newline at end of file
