UEBP-40 C++ code generation (WIP):
- code refactor
- inner structure properties are compared against default data

diff --git a/Engine/Source/Developer/BlueprintCompilerCppBackend/Private/BlueprintCompilerCppBackend.cpp b/Engine/Source/Developer/BlueprintCompilerCppBackend/Private/BlueprintCompilerCppBackend.cpp
--- a/Engine/Source/Developer/BlueprintCompilerCppBackend/Private/BlueprintCompilerCppBackend.cpp
+++ b/Engine/Source/Developer/BlueprintCompilerCppBackend/Private/BlueprintCompilerCppBackend.cpp
@@ -84,7 +84,7 @@ FString FBlueprintCompilerCppBackend::LatentFunctionInfoTermToText(FBPTerminal*
 	check(LinkageTermStartIdx != INDEX_NONE);
 	StructValues = StructValues.Replace(TEXT("-1"), *FString::FromInt(TargetStateIndex));
 
-	return FString(TEXT("F")) + LatentInfoStruct->GetName() + StructValues;
+	return FEmitHelper::LiteralTerm(Term->Type, StructValues, nullptr);
 }
 
 void FBlueprintCompilerCppBackend::EmitStructProperties(FStringOutputDevice& Target, UStruct* SourceClass)
diff --git a/Engine/Source/Developer/BlueprintCompilerCppBackend/Private/BlueprintCompilerCppBackendUtils.cpp b/Engine/Source/Developer/BlueprintCompilerCppBackend/Private/BlueprintCompilerCppBackendUtils.cpp
--- a/Engine/Source/Developer/BlueprintCompilerCppBackend/Private/BlueprintCompilerCppBackendUtils.cpp
+++ b/Engine/Source/Developer/BlueprintCompilerCppBackend/Private/BlueprintCompilerCppBackendUtils.cpp
@@ -540,6 +540,9 @@ FString FEmitHelper::LiteralTerm(const FEdGraphPinType& Type, const FString& Cus
 			const FString ConstructBody = bEmptyCustomValue
 				? (StructType->IsA<UUserDefinedStruct>() ? TEXT("::GetDefaultValue()") : TEXT("{}"))
 				: CustomValue;
+
+			// Import.. export..
+
 			return StructName + ConstructBody;
 		}
 	}
@@ -672,178 +675,6 @@ bool FEmitHelper::GenerateAssignmentCast(const FEdGraphPinType& LType, const FEd
 	return false;
 }
 
-struct FDefaultValueHelperContext
-{
-private:
-	FString Indent;
-	FString Result;
-
-	FString AfterCreateFormat;
-	FString ActualClassExpression;
-	UClass* ActualClass;
-
-	TMap<UObject*, FString> NativeObjectNames;
-
-	TArray<UObject*> OrderedObjects;
-
-	int32 LocalNameIndexMax;
-
-	TArray<FString> HighPriorityLines;
-	TArray<FString> LowPriorityLines;
-public:
-
-	bool bCreatingObjectsPerClass;
-
-	FDefaultValueHelperContext()
-		: ActualClass(nullptr)
-		, LocalNameIndexMax(0)
-		, bCreatingObjectsPerClass(false)
-	{}
-
-// OBJECTS:
-	bool FindLocalObject(UObject* Object, FString& OutNamePath, FString* OutAfterCreateFormatPtr = nullptr) const
-	{
-		if (Object == ActualClass)
-		{
-			if (OutAfterCreateFormatPtr)
-			{
-				*OutAfterCreateFormatPtr = AfterCreateFormat;
-			}
-			OutNamePath = ActualClassExpression;
-			return true;
-		}
-
-		if (const FString* NamePtr = NativeObjectNames.Find(Object))
-		{
-			OutNamePath = *NamePtr;
-			if (OutAfterCreateFormatPtr)
-			{
-				*OutAfterCreateFormatPtr = FString();
-			}
-			return true;
-		}
-
-		return false;
-	}
-
-	FString GenerateGetProperty(const UProperty* Property) const
-	{
-		check(Property);
-		const UStruct* OwnerStruct = Property->GetOwnerStruct();
-		const FString OwnerName = FString(OwnerStruct->GetPrefixCPP()) + OwnerStruct->GetName();
-		const FString OwnerPath = OwnerName + (OwnerStruct->IsA<UClass>() ? TEXT("::StaticClass()") : TEXT("::StaticStruct()"));
-		//return FString::Printf(TEXT("%s->FindPropertyByName(GET_MEMBER_NAME_CHECKED(%s, %s))"), *OwnerPath, *OwnerName, *Property->GetNameCPP());
-
-		return FString::Printf(TEXT("%s->FindPropertyByName(FName(TEXT(\"%s\")))"), *OwnerPath, *Property->GetNameCPP());
-	}
-
-	void SetCurrentlyGeneratedClass(UClass* InClass, const FString& AfterCreateFormatWhenTheObjectIsOuter, const FString& NamePath)
-	{
-		ActualClass = InClass;
-		ActualClassExpression = NamePath;
-		AfterCreateFormat = AfterCreateFormatWhenTheObjectIsOuter;
-	}
-
-	FString GenerateUniqueLocalName()
-	{
-		const FString UniqueNameBase = TEXT("__Local__");
-		const FString UniqueName = FString::Printf(TEXT("%s%d"), *UniqueNameBase, LocalNameIndexMax);
-		++LocalNameIndexMax;
-		return UniqueName;
-	}
-
-	void AddObjectFromLocalProperty(UObject* Object, const FString& NativeName)
-	{
-		NativeObjectNames.Add(Object, NativeName);
-	}
-
-	FString AddNewObject(UObject* Object)
-	{
-		check(!NativeObjectNames.Contains(Object));
-		const FString UniqueName = GenerateUniqueLocalName();
-		NativeObjectNames.Add(Object, UniqueName);
-		OrderedObjects.Add(Object);
-		return UniqueName;
-	}
-
-	UClass* GetCurrentlyGeneratedClass() const
-	{
-		return ActualClass;
-	}
-
-	const TArray<UObject*>& GetOrderedObjects()
-	{
-		return OrderedObjects;
-	}
-
-	FString FindGlobalObject(UObject* Object)
-	{
-		// TODO: check if not excluded
-
-		if (auto BPGC = ExactCast<UBlueprintGeneratedClass>(Object))
-		{
-			// TODO: check if supported 
-			return FString::Printf(TEXT("%s%s::StaticClass()"), BPGC->GetPrefixCPP(), *BPGC->GetName());
-		}
-
-		if (auto UDS = Cast<UUserDefinedStruct>(Object))
-		{
-			return FString::Printf(TEXT("%s%s::StaticStruct()"), UDS->GetPrefixCPP(), *UDS->GetName());
-		}
-
-		if (auto UDE = Cast<UUserDefinedEnum>(Object))
-		{
-			return FString::Printf(TEXT("%s_StaticEnum()"), *UDE->GetName());
-		}
-
-		// TODO: handle subobjects
-
-		return FString{};
-	}
-
-// TEXT
-	void IncreaseIndent()
-	{
-		Indent += TEXT("\t");
-	}
-
-	void DecreaseIndent()
-	{
-		Indent.RemoveFromEnd(TEXT("\t"));
-	}
-
-	void AddHighPriorityLine(const FString& Line)
-	{
-		HighPriorityLines.Emplace(FString::Printf(TEXT("%s%s\n"), *Indent, *Line));
-	}
-
-	void AddLowPriorityLine(const FString& Line)
-	{
-		LowPriorityLines.Emplace(FString::Printf(TEXT("%s%s\n"), *Indent, *Line));
-	}
-
-	void FlushLines()
-	{
-		for (auto& Line : HighPriorityLines)
-		{
-			Result += Line;
-		}
-		HighPriorityLines.Reset();
-
-		for (auto& Line : LowPriorityLines)
-		{
-			Result += Line;
-		}
-		LowPriorityLines.Reset();
-	}
-
-	FString GetResult()
-	{
-		FlushLines();
-		return Result;
-	}
-};
-
 void FEmitDefaultValueHelper::OuterGenerate(FDefaultValueHelperContext& Context
 	, const UProperty* Property
 	, const FString& OuterPath
@@ -878,7 +709,7 @@ void FEmitDefaultValueHelper::OuterGenerate(FDefaultValueHelperContext& Context
 				const FString OperatorStr = (EPropertyAccessOperator::Dot == AccessOperator) ? TEXT("&") : TEXT("");
 				const FString ContainerStr = (EPropertyAccessOperator::None == AccessOperator) ? TEXT("this") : FString::Printf(TEXT("%s(%s)"), *OperatorStr, *OuterPath);
 				PathToMember = Context.GenerateUniqueLocalName();
-				Context.AddLowPriorityLine(FString::Printf(TEXT("auto& %s = *(%s->ContainerPtrToValuePtr<%s>(%s, %d));")
+				Context.AddLine(FString::Printf(TEXT("auto& %s = *(%s->ContainerPtrToValuePtr<%s>(%s, %d));")
 					, *PathToMember, *PropertyObject, *TypeDeclaration, *ContainerStr, ArrayIndex));
 			}
 			else
@@ -890,7 +721,8 @@ void FEmitDefaultValueHelper::OuterGenerate(FDefaultValueHelperContext& Context
 				PathToMember = FString::Printf(TEXT("%s%s%s%s"), *OuterPath, *AccessOperatorStr, *Property->GetNameCPP(), *ArrayPost);
 			}
 			const uint8* ValuePtr = Property->ContainerPtrToValuePtr<uint8>(DataContainer, ArrayIndex);
-			InnerGenerate(Context, Property, PathToMember, ValuePtr);
+			const uint8* DefaultValuePtr = OptionalDefaultDataContainer ? Property->ContainerPtrToValuePtr<uint8>(OptionalDefaultDataContainer, ArrayIndex) : nullptr;
+			InnerGenerate(Context, Property, PathToMember, ValuePtr, DefaultValuePtr);
 		}
 	}
 }
@@ -916,9 +748,9 @@ FString FEmitDefaultValueHelper::GenerateGetDefaultValue(const UUserDefinedStruc
 	return Result;
 }
 
-void FEmitDefaultValueHelper::InnerGenerate(FDefaultValueHelperContext& Context, const UProperty* Property, const FString& PathToMember, const uint8* ValuePtr, bool bWithoutFirstConstructionLine)
+void FEmitDefaultValueHelper::InnerGenerate(FDefaultValueHelperContext& Context, const UProperty* Property, const FString& PathToMember, const uint8* ValuePtr, const uint8* DefaultValuePtr, bool bWithoutFirstConstructionLine)
 {
-	auto OneLineConstruction = [](FDefaultValueHelperContext& LocalContext, const UProperty* LocalProperty, const uint8* LocalValuePtr, FString& OutSingleLine) -> bool
+	auto OneLineConstruction = [](FDefaultValueHelperContext& LocalContext, const UProperty* LocalProperty, const uint8* LocalValuePtr, FString& OutSingleLine, bool bGenerateEmptyStructConstructor) -> bool
 	{
 		bool bComplete = true;
 		FString ValueStr = HandleSpecialTypes(LocalContext, LocalProperty, LocalValuePtr);
@@ -929,10 +761,13 @@ void FEmitDefaultValueHelper::InnerGenerate(FDefaultValueHelperContext& Context,
 			if (ValueStr.IsEmpty() && StructProperty)
 			{
 				check(StructProperty->Struct);
-				ValueStr = FString::Printf(TEXT("F%s()"), *StructProperty->Struct->GetName());
+				if (bGenerateEmptyStructConstructor)
+				{
+					ValueStr = FString::Printf(TEXT("F%s()"), *StructProperty->Struct->GetName()); //don;t override existing values
+				}
 				bComplete = false;
 			}
-			if (ValueStr.IsEmpty())
+			else if (ValueStr.IsEmpty())
 			{
 				UE_LOG(LogK2Compiler, Error, TEXT("FEmitDefaultValueHelper Cannot generate initilization: %s"), *LocalProperty->GetPathName());
 			}
@@ -949,8 +784,12 @@ void FEmitDefaultValueHelper::InnerGenerate(FDefaultValueHelperContext& Context,
 	if (!bWithoutFirstConstructionLine)
 	{
 		FString ValueStr;
-		const bool bComplete = OneLineConstruction(Context, Property, ValuePtr, ValueStr);
-		Context.AddLowPriorityLine(FString::Printf(TEXT("%s = %s;"), *PathToMember, *ValueStr));
+		const bool bComplete = OneLineConstruction(Context, Property, ValuePtr, ValueStr, false);
+		if (!ValueStr.IsEmpty())
+		{
+			Context.AddLine(FString::Printf(TEXT("%s = %s;"), *PathToMember, *ValueStr));
+		}
+		// array initialization "array_var = TArray<..>()" is complete, but it still needs items.
 		if (bComplete && !ArrayProperty)
 		{
 			return;
@@ -961,7 +800,7 @@ void FEmitDefaultValueHelper::InnerGenerate(FDefaultValueHelperContext& Context,
 	{
 		for (auto LocalProperty : TFieldRange<const UProperty>(StructProperty->Struct))
 		{
-			OuterGenerate(Context, LocalProperty, PathToMember, ValuePtr, nullptr, EPropertyAccessOperator::Dot);
+			OuterGenerate(Context, LocalProperty, PathToMember, ValuePtr, DefaultValuePtr, EPropertyAccessOperator::Dot);
 		}
 	}
 	
@@ -973,12 +812,12 @@ void FEmitDefaultValueHelper::InnerGenerate(FDefaultValueHelperContext& Context,
 			const uint8* LocalValuePtr = ScriptArrayHelper.GetRawPtr(Index);
 
 			FString ValueStr;
-			const bool bComplete = OneLineConstruction(Context, ArrayProperty->Inner, LocalValuePtr, ValueStr);
-			Context.AddLowPriorityLine(FString::Printf(TEXT("%s.Add(%s);"), *PathToMember, *ValueStr));
+			const bool bComplete = OneLineConstruction(Context, ArrayProperty->Inner, LocalValuePtr, ValueStr, true);
+			Context.AddLine(FString::Printf(TEXT("%s.Add(%s);"), *PathToMember, *ValueStr));
 			if (!bComplete)
 			{
 				const FString LocalPathToMember = FString::Printf(TEXT("%s[%d]"), *PathToMember, Index);
-				InnerGenerate(Context, ArrayProperty->Inner, LocalPathToMember, LocalValuePtr, true);
+				InnerGenerate(Context, ArrayProperty->Inner, LocalPathToMember, LocalValuePtr, nullptr, true);
 			}
 		}
 	}
@@ -993,7 +832,7 @@ FString FEmitDefaultValueHelper::HandleSpecialTypes(FDefaultValueHelperContext&
 		if (Object)
 		{
 			FString NativeName;
-			if (Context.FindLocalObject(Object, NativeName))
+			if (Context.FindLocalObject_InConstructor(Object, NativeName))
 			{
 				return NativeName;
 			}
@@ -1078,10 +917,9 @@ UActorComponent* FEmitDefaultValueHelper::HandleNonNativeComponent(FDefaultValue
 	auto ComponentClass = ComponentTemplate->GetClass();
 	if (bNew)
 	{
-		Context.AddLowPriorityLine(TEXT(""));
 		Context.AddHighPriorityLine(FString::Printf(TEXT("%s = CreateDefaultSubobject<%s%s>(TEXT(\"%s\"));")
 			, *NativeName, ComponentClass->GetPrefixCPP(), *ComponentClass->GetName(), *ComponentTemplate->GetName()));
-		Context.AddObjectFromLocalProperty(ComponentTemplate, NativeName);
+		Context.AddObjectFromLocalProperty_InConstructor(ComponentTemplate, NativeName);
 	}
 
 	UObject* ParentTemplateComponent = bNew 
@@ -1102,6 +940,42 @@ UActorComponent* FEmitDefaultValueHelper::HandleNonNativeComponent(FDefaultValue
 			, EPropertyAccessOperator::Pointer);
 	}
 
+	auto FindParentComponentName = [&]() -> FString
+	{
+		FString ResultName;
+		if (auto ParentNode = Node->GetSCS()->FindParentNode(Node))
+		{
+			const FName ParentVariableName = ParentNode->GetVariableName();
+			ResultName = (ParentVariableName != NAME_None) ? ParentVariableName.ToString() : FString();
+		}
+
+		if (ResultName.IsEmpty())
+		{
+			auto ParentTemplate = Node->GetParentComponentTemplate(CastChecked<UBlueprint>(BPGC->ClassGeneratedBy));
+			FString ParentComponentNativeName;
+			if (ParentTemplate && Context.FindLocalObject_InConstructor(ParentTemplateComponent, ParentComponentNativeName))
+			{
+				ResultName = ParentComponentNativeName;
+			}
+		}
+
+		if (ResultName.IsEmpty())
+		{
+			ResultName = (Node->ParentComponentOrVariableName != NAME_None) ? Node->ParentComponentOrVariableName.ToString() : FString();
+		}
+
+		return ResultName;
+	};
+
+	const FString ParentNodeName = FindParentComponentName();
+	if (!ParentNodeName.IsEmpty())
+	{
+		Context.AddLine(FString::Printf(TEXT("%s->AttachParent = %s;"), *NativeName, *ParentNodeName));
+	}
+
+	// TODO: Set root component
+	// TODO: Update local transform, etc..
+
 	return ComponentTemplate;
 }
 
@@ -1111,8 +985,9 @@ FString FEmitDefaultValueHelper::GenerateConstructor(UClass* InBPGC)
 	const FString CppClassName = FString(BPGC->GetPrefixCPP()) + BPGC->GetName();
 
 	FDefaultValueHelperContext Context;
-	Context.AddLowPriorityLine(FString::Printf(TEXT("%s::%s(const FObjectInitializer& ObjectInitializer) : Super(ObjectInitializer)"), *CppClassName, *CppClassName));
-	Context.AddLowPriorityLine(TEXT("{"));
+	Context.SetCurrentlyGeneratedClass(BPGC);
+	Context.AddLine(FString::Printf(TEXT("%s::%s(const FObjectInitializer& ObjectInitializer) : Super(ObjectInitializer)"), *CppClassName, *CppClassName));
+	Context.AddLine(TEXT("{"));
 	Context.IncreaseIndent();
 
 	// When CDO is created create all subobjects owned by the class
@@ -1136,12 +1011,11 @@ FString FEmitDefaultValueHelper::GenerateConstructor(UClass* InBPGC)
 			});
 		}
 
-		Context.AddLowPriorityLine(FString::Printf(TEXT("if(HasAnyFlags(RF_ClassDefaultObject) && (%s::StaticClass() == GetClass()))"), *CppClassName));
-		Context.AddLowPriorityLine(TEXT("{"));
+		Context.AddLine(FString::Printf(TEXT("if(HasAnyFlags(RF_ClassDefaultObject) && (%s::StaticClass() == GetClass()))"), *CppClassName));
+		Context.AddLine(TEXT("{"));
 		Context.IncreaseIndent();
-		Context.AddLowPriorityLine(TEXT("ensure(0 == GetClass()->MiscObjects.Num());"));
+		Context.AddLine(TEXT("ensure(0 == GetClass()->MiscObjects.Num());"));
 		Context.bCreatingObjectsPerClass = true;
-		Context.SetCurrentlyGeneratedClass(BPGC, TEXT("GetClass()->MiscObjects.Add(%s);"), TEXT("GetClass()"));
 		Context.FlushLines();
 		for (auto ComponentTemplate : ActorComponentTempatesOwnedByClass)
 		{
@@ -1169,17 +1043,17 @@ FString FEmitDefaultValueHelper::GenerateConstructor(UClass* InBPGC)
 
 		Context.bCreatingObjectsPerClass = false;
 		Context.DecreaseIndent();
-		Context.AddLowPriorityLine(TEXT("}"));
+		Context.AddLine(TEXT("}"));
 
 		// Let's have an easy access to generated class subobjects
-		Context.AddLowPriorityLine(TEXT("{")); // no shadow variables
+		Context.AddLine(TEXT("{")); // no shadow variables
 		Context.IncreaseIndent();
 		Context.FlushLines();
 		int32 SubobjectIndex = 0;
-		for (UObject* Obj : Context.GetOrderedObjects())
+		for (UObject* Obj : Context.GetObjectsCreatedPerClass_InConstructor())
 		{
 			FString NativeName;
-			const bool bFound = Context.FindLocalObject(Obj, NativeName);
+			const bool bFound = Context.FindLocalObject_InConstructor(Obj, NativeName);
 			ensure(bFound);
 			Context.AddHighPriorityLine(FString::Printf(TEXT("auto %s = CastChecked<%s%s>(GetClass()->MiscObjects[%d]);")
 				, *NativeName
@@ -1191,11 +1065,11 @@ FString FEmitDefaultValueHelper::GenerateConstructor(UClass* InBPGC)
 	}
 
 	UObject* CDO = BPGC->GetDefaultObject(false);
-	Context.AddObjectFromLocalProperty(CDO, TEXT("this"));
+	Context.AddObjectFromLocalProperty_InConstructor(CDO, TEXT("this"));
 
 	UObject* ParentCDO = BPGC->GetSuperClass()->GetDefaultObject(false);
 	check(CDO && ParentCDO);
-	Context.AddLowPriorityLine(TEXT(""));
+	Context.AddLine(TEXT(""));
 
 	// TIMELINES
 	TSet<const UProperty*> HandledProperties;
@@ -1227,9 +1101,9 @@ FString FEmitDefaultValueHelper::GenerateConstructor(UClass* InBPGC)
 	}
 
 	Context.DecreaseIndent();
-	Context.AddLowPriorityLine(TEXT("}"));
+	Context.AddLine(TEXT("}"));
 	Context.DecreaseIndent();
-	Context.AddLowPriorityLine(TEXT("}"));
+	Context.AddLine(TEXT("}"));
 
 	return Context.GetResult();
 }
@@ -1237,20 +1111,22 @@ FString FEmitDefaultValueHelper::GenerateConstructor(UClass* InBPGC)
 FString FEmitDefaultValueHelper::HandleClassSubobject(FDefaultValueHelperContext& Context, UObject* Object)
 {
 	FString OuterStr;
-	FString AfterCreateFormat;
-	if (!Context.FindLocalObject(Object->GetOuter(), OuterStr, &AfterCreateFormat))
+	if (!Context.FindLocalObject_InConstructor(Object->GetOuter(), OuterStr))
 	{
 		return FString();
 	}
 
-	const FString LocalNativeName = Context.AddNewObject(Object);
+	const FString LocalNativeName = Context.AddNewObject_InConstructor(Object);
 	UClass* ObjectClass = Object->GetClass();
-	Context.AddHighPriorityLine(FString::Printf(TEXT("auto %s = NewObject<%s%s>(%s, TEXT(\"%s\"));")
-		, *LocalNativeName, ObjectClass->GetPrefixCPP(), *ObjectClass->GetName()
+	Context.AddHighPriorityLine(FString::Printf(
+		TEXT("auto %s = NewObject<%s%s>(%s, TEXT(\"%s\"));")
+		, *LocalNativeName
+		, ObjectClass->GetPrefixCPP()
+		, *ObjectClass->GetName()
 		, *OuterStr, *Object->GetName()));
-	if (!AfterCreateFormat.IsEmpty())
+	if (Object->GetOuter() == Context.GetCurrentlyGeneratedClass())
 	{
-		Context.AddLowPriorityLine(FString::Printf(*AfterCreateFormat, *LocalNativeName));
+		Context.AddLine(FString::Printf(TEXT("GetClass()->MiscObjects.Add(%s);"), *LocalNativeName));
 	}
 
 	for (auto Property : TFieldRange<const UProperty>(ObjectClass))
@@ -1268,7 +1144,7 @@ FString FEmitDefaultValueHelper::HandleInstancedSubobject(FDefaultValueHelperCon
 {
 	check(Object);
 
-	const FString LocalNativeName = Context.AddNewObject(Object);
+	const FString LocalNativeName = Context.AddNewObject_InConstructor(Object);
 	UClass* ObjectClass = Object->GetClass();
 
 	auto BPGC = Context.GetCurrentlyGeneratedClass();
@@ -1281,7 +1157,7 @@ FString FEmitDefaultValueHelper::HandleInstancedSubobject(FDefaultValueHelperCon
 	else
 	{
 		FString OuterStr;
-		if (!Context.FindLocalObject(Object->GetOuter(), OuterStr))
+		if (!Context.FindLocalObject_InConstructor(Object->GetOuter(), OuterStr))
 		{
 			return FString();
 		}
diff --git a/Engine/Source/Developer/BlueprintCompilerCppBackend/Private/BlueprintCompilerCppBackendUtils.h b/Engine/Source/Developer/BlueprintCompilerCppBackend/Private/BlueprintCompilerCppBackendUtils.h
--- a/Engine/Source/Developer/BlueprintCompilerCppBackend/Private/BlueprintCompilerCppBackendUtils.h
+++ b/Engine/Source/Developer/BlueprintCompilerCppBackend/Private/BlueprintCompilerCppBackendUtils.h
@@ -65,7 +65,192 @@ struct FEmitHelper
 	static bool GenerateAssignmentCast(const FEdGraphPinType& LType, const FEdGraphPinType& RType, FString& OutCastBegin, FString& OutCastEnd);
 };
 
-struct FDefaultValueHelperContext;
+struct FDefaultValueHelperContext
+{
+private:
+	FString Indent;
+	FString Result;
+	int32 LocalNameIndexMax;
+	TArray<FString> HighPriorityLines;
+	TArray<FString> LowPriorityLines;
+	UClass* ActualClass;
+
+	TMap<UObject*, FString> NativeObjectNamesInConstructor;
+	TArray<UObject*> ObjectsCreatedPerClass;
+
+
+public:
+
+	bool bCreatingObjectsPerClass;
+
+	FDefaultValueHelperContext()
+		: ActualClass(nullptr)
+		, LocalNameIndexMax(0)
+		, bCreatingObjectsPerClass(false)
+	{}
+
+	// Functions to use in constructor only
+	bool FindLocalObject_InConstructor(UObject* Object, FString& OutNamePath) const
+	{
+		if (Object == ActualClass)
+		{
+			OutNamePath = TEXT("GetClass()");
+			return true;
+		}
+
+		if (const FString* NamePtr = NativeObjectNamesInConstructor.Find(Object))
+		{
+			OutNamePath = *NamePtr;
+			return true;
+		}
+
+		return false;
+	}
+
+	void AddObjectFromLocalProperty_InConstructor(UObject* Object, const FString& NativeName)
+	{
+		NativeObjectNamesInConstructor.Add(Object, NativeName);
+	}
+
+	FString AddNewObject_InConstructor(UObject* Object)
+	{
+		const FString UniqueName = GenerateUniqueLocalName();
+		NativeObjectNamesInConstructor.Add(Object, UniqueName);
+		if (bCreatingObjectsPerClass)
+		{
+			check(!ObjectsCreatedPerClass.Contains(Object));
+			ObjectsCreatedPerClass.Add(Object);
+		}
+
+		return UniqueName;
+	}
+
+	const TArray<UObject*>& GetObjectsCreatedPerClass_InConstructor()
+	{
+		return ObjectsCreatedPerClass;
+	}
+
+	// Universal functions
+	FString GenerateGetProperty(const UProperty* Property) const
+	{
+		check(Property);
+		const UStruct* OwnerStruct = Property->GetOwnerStruct();
+		const FString OwnerName = FString(OwnerStruct->GetPrefixCPP()) + OwnerStruct->GetName();
+		const FString OwnerPath = OwnerName + (OwnerStruct->IsA<UClass>() ? TEXT("::StaticClass()") : TEXT("::StaticStruct()"));
+		//return FString::Printf(TEXT("%s->FindPropertyByName(GET_MEMBER_NAME_CHECKED(%s, %s))"), *OwnerPath, *OwnerName, *Property->GetNameCPP());
+
+		return FString::Printf(TEXT("%s->FindPropertyByName(FName(TEXT(\"%s\")))"), *OwnerPath, *Property->GetNameCPP());
+	}
+
+	void SetCurrentlyGeneratedClass(UClass* InClass)
+	{
+		ActualClass = InClass;
+	}
+
+	FString GenerateUniqueLocalName()
+	{
+		const FString UniqueNameBase = TEXT("__Local__");
+		const FString UniqueName = FString::Printf(TEXT("%s%d"), *UniqueNameBase, LocalNameIndexMax);
+		++LocalNameIndexMax;
+		return UniqueName;
+	}
+
+	UClass* GetCurrentlyGeneratedClass() const
+	{
+		return ActualClass;
+	}
+
+	FString FindGlobalObject(UObject* Object)
+	{
+		// TODO: check if not excluded
+
+		if (ActualClass && (Object == ActualClass))
+		{
+			return TEXT("GetClass()");
+		}
+
+		if (auto ObjClass = Cast<UClass>(Object))
+		{
+			if (ObjClass->HasAnyClassFlags(CLASS_Native))
+			{
+				return FString::Printf(TEXT("%s%s::StaticClass()"), ObjClass->GetPrefixCPP(), *ObjClass->GetName());
+			}
+		}
+
+		// TODO Handle native structires, and special cases..
+
+		if (auto BPGC = ExactCast<UBlueprintGeneratedClass>(Object))
+		{
+			// TODO: check if supported 
+			return FString::Printf(TEXT("%s%s::StaticClass()"), BPGC->GetPrefixCPP(), *BPGC->GetName());
+		}
+
+		if (auto UDS = Cast<UScriptStruct>(Object))
+		{
+			// Check if  
+			// TODO: check if supported 
+			return FString::Printf(TEXT("%s%s::StaticStruct()"), UDS->GetPrefixCPP(), *UDS->GetName());
+		}
+
+		if (auto UDE = Cast<UEnum>(Object))
+		{
+			// TODO: check if supported 
+			return FString::Printf(TEXT("%s_StaticEnum()"), *UDE->GetName());
+		}
+
+		// TODO: handle subobjects
+
+		return FString{};
+	}
+
+	// TEXT
+	void IncreaseIndent()
+	{
+		Indent += TEXT("\t");
+	}
+
+	void DecreaseIndent()
+	{
+		Indent.RemoveFromEnd(TEXT("\t"));
+	}
+
+	void AddHighPriorityLine(const FString& Line)
+	{
+		HighPriorityLines.Emplace(FString::Printf(TEXT("%s%s\n"), *Indent, *Line));
+	}
+
+	void AddLine(const FString& Line)
+	{
+		LowPriorityLines.Emplace(FString::Printf(TEXT("%s%s\n"), *Indent, *Line));
+	}
+
+	void FlushLines()
+	{
+		for (auto& Line : HighPriorityLines)
+		{
+			Result += Line;
+		}
+		HighPriorityLines.Reset();
+
+		for (auto& Line : LowPriorityLines)
+		{
+			Result += Line;
+		}
+		LowPriorityLines.Reset();
+	}
+
+	FString GetResult()
+	{
+		FlushLines();
+		return Result;
+	}
+
+	void ClearResult()
+	{
+		FlushLines();
+		Result.Reset();
+	}
+};
 
 struct FEmitDefaultValueHelper
 {
@@ -85,7 +270,7 @@ private:
 	static void OuterGenerate(FDefaultValueHelperContext& Context, const UProperty* Property, const FString& OuterPath, const uint8* DataContainer, const uint8* OptionalDefaultDataContainer, EPropertyAccessOperator AccessOperator, bool bAllowProtected = false);
 	
 	// PathToMember ends with variable name
-	static void InnerGenerate(FDefaultValueHelperContext& Context, const UProperty* Property, const FString& PathToMember, const uint8* ValuePtr, bool bWithoutFirstConstructionLine = false);
+	static void InnerGenerate(FDefaultValueHelperContext& Context, const UProperty* Property, const FString& PathToMember, const uint8* ValuePtr, const uint8* DefaultValuePtr, bool bWithoutFirstConstructionLine = false);
 	
 	// Returns native term, 
 	// returns empty string if cannot handle
