UClass::MiscObjects -> UClass::ConvertedSubobjectsFromBPGC

diff --git a/Engine/Source/Developer/BlueprintCompilerCppBackend/Private/BlueprintCompilerCppBackendUtils.cpp b/Engine/Source/Developer/BlueprintCompilerCppBackend/Private/BlueprintCompilerCppBackendUtils.cpp
--- a/Engine/Source/Developer/BlueprintCompilerCppBackend/Private/BlueprintCompilerCppBackendUtils.cpp
+++ b/Engine/Source/Developer/BlueprintCompilerCppBackend/Private/BlueprintCompilerCppBackendUtils.cpp
@@ -1014,7 +1014,7 @@ FString FEmitDefaultValueHelper::GenerateConstructor(UClass* InBPGC)
 		Context.AddLine(FString::Printf(TEXT("if(HasAnyFlags(RF_ClassDefaultObject) && (%s::StaticClass() == GetClass()))"), *CppClassName));
 		Context.AddLine(TEXT("{"));
 		Context.IncreaseIndent();
-		Context.AddLine(TEXT("ensure(0 == GetClass()->MiscObjects.Num());"));
+		Context.AddLine(TEXT("ensure(0 == GetClass()->ConvertedSubobjectsFromBPGC.Num());"));
 		Context.bCreatingObjectsPerClass = true;
 		Context.FlushLines();
 		for (auto ComponentTemplate : ActorComponentTempatesOwnedByClass)
@@ -1055,7 +1055,7 @@ FString FEmitDefaultValueHelper::GenerateConstructor(UClass* InBPGC)
 			FString NativeName;
 			const bool bFound = Context.FindLocalObject_InConstructor(Obj, NativeName);
 			ensure(bFound);
-			Context.AddHighPriorityLine(FString::Printf(TEXT("auto %s = CastChecked<%s%s>(GetClass()->MiscObjects[%d]);")
+			Context.AddHighPriorityLine(FString::Printf(TEXT("auto %s = CastChecked<%s%s>(GetClass()->ConvertedSubobjectsFromBPGC[%d]);")
 				, *NativeName
 				, Obj->GetClass()->GetPrefixCPP()
 				, *Obj->GetClass()->GetName()
@@ -1126,7 +1126,7 @@ FString FEmitDefaultValueHelper::HandleClassSubobject(FDefaultValueHelperContext
 		, *OuterStr, *Object->GetName()));
 	if (Object->GetOuter() == Context.GetCurrentlyGeneratedClass())
 	{
-		Context.AddLine(FString::Printf(TEXT("GetClass()->MiscObjects.Add(%s);"), *LocalNativeName));
+		Context.AddLine(FString::Printf(TEXT("GetClass()->ConvertedSubobjectsFromBPGC.Add(%s);"), *LocalNativeName));
 	}
 
 	for (auto Property : TFieldRange<const UProperty>(ObjectClass))
diff --git a/Engine/Source/Runtime/CoreUObject/Private/UObject/Class.cpp b/Engine/Source/Runtime/CoreUObject/Private/UObject/Class.cpp
--- a/Engine/Source/Runtime/CoreUObject/Private/UObject/Class.cpp
+++ b/Engine/Source/Runtime/CoreUObject/Private/UObject/Class.cpp
@@ -2525,7 +2525,7 @@ void UClass::AddReferencedObjects(UObject* InThis, FReferenceCollector& Collecto
 		This->CallAddReferencedObjects(This->ClassDefaultObject, Collector);
 	}
 
-	Collector.AddReferencedObjects(This->MiscObjects, This);
+	Collector.AddReferencedObjects(This->ConvertedSubobjectsFromBPGC, This);
 
 	Super::AddReferencedObjects( This, Collector );
 }
diff --git a/Engine/Source/Runtime/CoreUObject/Public/UObject/Class.h b/Engine/Source/Runtime/CoreUObject/Public/UObject/Class.h
--- a/Engine/Source/Runtime/CoreUObject/Public/UObject/Class.h
+++ b/Engine/Source/Runtime/CoreUObject/Public/UObject/Class.h
@@ -1934,8 +1934,8 @@ public:
 	/** This class's native functions. */
 	TArray<FNativeFunctionLookup> NativeFunctionLookupTable;
 
-	/** Misc Objects Owned by the class. Mostly templates from BP-generated native class */
-	TArray<UObject*> MiscObjects;
+	/** Misc objects owned by the class. Mostly templates from BP-generated native class */
+	TArray<UObject*> ConvertedSubobjectsFromBPGC;
 
 public:
 	// Constructors
diff --git a/Engine/Source/Runtime/Engine/Private/ActorConstruction.cpp b/Engine/Source/Runtime/Engine/Private/ActorConstruction.cpp
--- a/Engine/Source/Runtime/Engine/Private/ActorConstruction.cpp
+++ b/Engine/Source/Runtime/Engine/Private/ActorConstruction.cpp
@@ -621,7 +621,7 @@ UActorComponent* AActor::AddComponent(FName TemplateName, bool bManualAttachment
 		}
 		else
 		{
-			UObject** FoundTemplatePtr = TemplateOwnerClass->MiscObjects.FindByPredicate([=](UObject* Obj) -> bool
+			UObject** FoundTemplatePtr = TemplateOwnerClass->ConvertedSubobjectsFromBPGC.FindByPredicate([=](UObject* Obj) -> bool
 			{
 				return Obj && Obj->IsA<UActorComponent>() && (Obj->GetFName() == TemplateName);
 			});
diff --git a/Engine/Source/Runtime/Engine/Private/BlueprintGeneratedClass.cpp b/Engine/Source/Runtime/Engine/Private/BlueprintGeneratedClass.cpp
--- a/Engine/Source/Runtime/Engine/Private/BlueprintGeneratedClass.cpp
+++ b/Engine/Source/Runtime/Engine/Private/BlueprintGeneratedClass.cpp
@@ -345,7 +345,7 @@ UDynamicBlueprintBinding* UBlueprintGeneratedClass::GetDynamicBindingObject(cons
 	}
 	else
 	{
-		for (auto MiscObj : ThisClass->MiscObjects)
+		for (auto MiscObj : ThisClass->ConvertedSubobjectsFromBPGC)
 		{
 			auto DynamicBindingObject = Cast<UDynamicBlueprintBinding>(MiscObj);
 			if (DynamicBindingObject && (DynamicBindingObject->GetClass() == BindingClass))
@@ -379,7 +379,7 @@ void UBlueprintGeneratedClass::BindDynamicDelegates(const UClass* ThisClass, UOb
 	}
 	else
 	{
-		for (auto MiscObj : ThisClass->MiscObjects)
+		for (auto MiscObj : ThisClass->ConvertedSubobjectsFromBPGC)
 		{
 			auto DynamicBindingObject = Cast<UDynamicBlueprintBinding>(MiscObj);
 			if (DynamicBindingObject)
@@ -570,7 +570,7 @@ void UBlueprintGeneratedClass::CreateComponentsForActor(const UClass* ThisClass,
 	}
 	else
 	{
-		for (auto MiscObj : ThisClass->MiscObjects)
+		for (auto MiscObj : ThisClass->ConvertedSubobjectsFromBPGC)
 		{
 			auto TimelineTemplate = Cast<const UTimelineTemplate>(MiscObj);
 			// Not fatal if NULL, but shouldn't happen and ignored if not wired up in graph
diff --git a/Engine/Source/Runtime/Engine/Private/InputDelegateBinding.cpp b/Engine/Source/Runtime/Engine/Private/InputDelegateBinding.cpp
--- a/Engine/Source/Runtime/Engine/Private/InputDelegateBinding.cpp
+++ b/Engine/Source/Runtime/Engine/Private/InputDelegateBinding.cpp
@@ -18,7 +18,7 @@ UInputDelegateBinding::UInputDelegateBinding(const FObjectInitializer& ObjectIni
 bool UInputDelegateBinding::SupportsInputDelegate(const UClass* InClass)
 {
 	return InClass && (
-		InClass->MiscObjects.Num() 
+		InClass->ConvertedSubobjectsFromBPGC.Num()
 		|| Cast<UBlueprintGeneratedClass>(InClass)
 		|| SupportsInputDelegate(InClass->GetSuperClass()));
 }
