[BuildWatcher] Fix for RootEditor_OnMac

Initialization order issue with FDefaultValueHelperContext properties

diff --git a/Engine/Source/Developer/BlueprintCompilerCppBackend/Private/BlueprintCompilerCppBackendUtils.h b/Engine/Source/Developer/BlueprintCompilerCppBackend/Private/BlueprintCompilerCppBackendUtils.h
--- a/Engine/Source/Developer/BlueprintCompilerCppBackend/Private/BlueprintCompilerCppBackendUtils.h
+++ b/Engine/Source/Developer/BlueprintCompilerCppBackend/Private/BlueprintCompilerCppBackendUtils.h
@@ -84,8 +84,8 @@ public:
 	bool bCreatingObjectsPerClass;
 
 	FDefaultValueHelperContext()
-		: ActualClass(nullptr)
-		, LocalNameIndexMax(0)
+		: LocalNameIndexMax(0)
+		, ActualClass(nullptr)
 		, bCreatingObjectsPerClass(false)
 	{}
 
