UEBP-40 C++ code generation (WIP):
- code refactor
- Latent info support
- Literal structures are properly generated

diff --git a/Engine/Source/Developer/BlueprintCompilerCppBackend/Private/BlueprintCompilerCppBackend.cpp b/Engine/Source/Developer/BlueprintCompilerCppBackend/Private/BlueprintCompilerCppBackend.cpp
--- a/Engine/Source/Developer/BlueprintCompilerCppBackend/Private/BlueprintCompilerCppBackend.cpp
+++ b/Engine/Source/Developer/BlueprintCompilerCppBackend/Private/BlueprintCompilerCppBackend.cpp
@@ -5,22 +5,22 @@
 #include "BlueprintCompilerCppBackendUtils.h"
 #include "EdGraphSchema_K2.h"
 
-FString FBlueprintCompilerCppBackend::TermToText(const FBPTerminal* Term, bool bUseSafeContext)
+FString FBlueprintCompilerCppBackend::TermToText(FEmitterLocalContext& EmitterContext, const FBPTerminal* Term, bool bUseSafeContext)
 {
 	const FString PSC_Self(TEXT("self"));
 	if (Term->bIsLiteral)
 	{
-		return FEmitHelper::LiteralTerm(Term->Type, Term->Name, Term->ObjectLiteral);
+		return FEmitHelper::LiteralTerm(EmitterContext, Term->Type, Term->Name, Term->ObjectLiteral);
 	}
 	else if (Term->InlineGeneratedParameter)
 	{
 		if (KCST_SwitchValue == Term->InlineGeneratedParameter->Type)
 		{
-			return EmitSwitchValueStatmentInner(*Term->InlineGeneratedParameter);
+			return EmitSwitchValueStatmentInner(EmitterContext, *Term->InlineGeneratedParameter);
 		}
 		else if (KCST_CallFunction == Term->InlineGeneratedParameter->Type)
 		{
-			return EmitCallStatmentInner(*Term->InlineGeneratedParameter, true);
+			return EmitCallStatmentInner(EmitterContext, *Term->InlineGeneratedParameter, true);
 		}
 		else
 		{
@@ -33,7 +33,7 @@ FString FBlueprintCompilerCppBackend::TermToText(const FBPTerminal* Term, bool b
 		FString ResultPath(TEXT(""));
 		if ((Term->Context != NULL) && (Term->Context->Name != PSC_Self))
 		{
-			ResultPath = TermToText(Term->Context, false);
+			ResultPath = TermToText(EmitterContext, Term->Context, false);
 
 			if (Term->Context->IsStructContextType())
 			{
@@ -48,17 +48,17 @@ FString FBlueprintCompilerCppBackend::TermToText(const FBPTerminal* Term, bool b
 		FString Conditions;
 		if (bUseSafeContext)
 		{
-			Conditions = FSafeContextScopedEmmitter::ValidationChain(Term->Context, *this);
+			Conditions = FSafeContextScopedEmmitter::ValidationChain(EmitterContext, Term->Context, *this);
 		}
 
 		ResultPath += Term->AssociatedVarProperty ? Term->AssociatedVarProperty->GetNameCPP() : Term->Name;
 		return Conditions.IsEmpty()
 			? ResultPath
-			: FString::Printf(TEXT("((%s) ? (%s) : (%s))"), *Conditions, *ResultPath, *FEmitHelper::DefaultValue(Term->Type));
+			: FString::Printf(TEXT("((%s) ? (%s) : (%s))"), *Conditions, *ResultPath, *FEmitHelper::DefaultValue(EmitterContext, Term->Type));
 	}
 }
 
-FString FBlueprintCompilerCppBackend::LatentFunctionInfoTermToText(FBPTerminal* Term, FBlueprintCompiledStatement* TargetLabel)
+FString FBlueprintCompilerCppBackend::LatentFunctionInfoTermToText(FEmitterLocalContext& EmitterContext, FBPTerminal* Term, FBlueprintCompiledStatement* TargetLabel)
 {
 	auto LatentInfoStruct = FLatentActionInfo::StaticStruct();
 
@@ -84,7 +84,7 @@ FString FBlueprintCompilerCppBackend::LatentFunctionInfoTermToText(FBPTerminal*
 	check(LinkageTermStartIdx != INDEX_NONE);
 	StructValues = StructValues.Replace(TEXT("-1"), *FString::FromInt(TargetStateIndex));
 
-	return FEmitHelper::LiteralTerm(Term->Type, StructValues, nullptr);
+	return FEmitHelper::LiteralTerm(EmitterContext, Term->Type, StructValues, nullptr);
 }
 
 void FBlueprintCompilerCppBackend::EmitStructProperties(FStringOutputDevice& Target, UStruct* SourceClass)
@@ -223,17 +223,14 @@ void FBlueprintCompilerCppBackend::GenerateCodeFromClass(UClass* SourceClass, TI
 	Emit(Body, *FEmitHelper::EmitLifetimeReplicatedPropsImpl(SourceClass, CppClassName, TEXT("")));
 }
 
-void FBlueprintCompilerCppBackend::EmitCallDelegateStatment(FKismetFunctionContext& FunctionContext, FBlueprintCompiledStatement& Statement)
+void FBlueprintCompilerCppBackend::EmitCallDelegateStatment(FEmitterLocalContext& EmitterContext, FKismetFunctionContext& FunctionContext, FBlueprintCompiledStatement& Statement)
 {
 	check(Statement.FunctionContext && Statement.FunctionContext->AssociatedVarProperty);
-	FSafeContextScopedEmmitter SafeContextScope(Body, Statement.FunctionContext->Context, *this, TEXT("\t\t\t"));
-
-	Emit(Body, TEXT("\t\t\t"));
-	Emit(Body, *SafeContextScope.GetAdditionalIndent());
-	Emit(Body, *FString::Printf(TEXT("%s.Broadcast(%s);"), *TermToText(Statement.FunctionContext, false), *EmitMethodInputParameterList(Statement)));
+	FSafeContextScopedEmmitter SafeContextScope(EmitterContext, Statement.FunctionContext->Context, *this);
+	EmitterContext.AddLine(FString::Printf(TEXT("%s.Broadcast(%s);"), *TermToText(EmitterContext, Statement.FunctionContext, false), *EmitMethodInputParameterList(EmitterContext, Statement)));
 }
 
-FString FBlueprintCompilerCppBackend::EmitMethodInputParameterList(FBlueprintCompiledStatement& Statement)
+FString FBlueprintCompilerCppBackend::EmitMethodInputParameterList(FEmitterLocalContext& EmitterContext, FBlueprintCompiledStatement& Statement)
 {
 	FString Result;
 	int32 NumParams = 0;
@@ -261,7 +258,7 @@ FString FBlueprintCompilerCppBackend::EmitMethodInputParameterList(FBlueprintCom
 				if (StructProp && (StructProp->Struct == FLatentActionInfo::StaticStruct()))
 				{
 					// Latent function info case
-					VarName = LatentFunctionInfoTermToText(Term, Statement.TargetLabel);
+					VarName = LatentFunctionInfoTermToText(EmitterContext, Term, Statement.TargetLabel);
 				}
 				else
 				{
@@ -272,7 +269,7 @@ FString FBlueprintCompilerCppBackend::EmitMethodInputParameterList(FBlueprintCom
 			else
 			{
 				// Emit a normal parameter term
-				VarName = TermToText(Term);
+				VarName = TermToText(EmitterContext, Term);
 			}
 
 			if (FuncParamProperty->HasAnyPropertyFlags(CPF_OutParm))
@@ -288,103 +285,103 @@ FString FBlueprintCompilerCppBackend::EmitMethodInputParameterList(FBlueprintCom
 	return Result;
 }
 
-FString FBlueprintCompilerCppBackend::EmitCallStatmentInner(FBlueprintCompiledStatement& Statement, bool bInline)
+FString FBlueprintCompilerCppBackend::EmitCallStatmentInner(FEmitterLocalContext& EmitterContext, FBlueprintCompiledStatement& Statement, bool bInline)
 {
-	FString Result;
-
 	const bool bCallOnDifferentObject = Statement.FunctionContext && (Statement.FunctionContext->Name != TEXT("self"));
 	const bool bStaticCall = Statement.FunctionToCall->HasAnyFunctionFlags(FUNC_Static);
 	const bool bUseSafeContext = bCallOnDifferentObject && !bStaticCall;
 	const bool bInterfaceCall = bCallOnDifferentObject && Statement.FunctionContext && (UEdGraphSchema_K2::PC_Interface == Statement.FunctionContext->Type.PinCategory);
-	{
-		FSafeContextScopedEmmitter SafeContextScope(Result, bUseSafeContext ? Statement.FunctionContext : nullptr, *this, TEXT("\t\t\t"));
-		ensure(!bInline || !SafeContextScope.IsSafeContextUsed());
 
-		bool bCloseCast = false;
-		if (!bInline)
+	FString Result;
+	bool bCloseCast = false;
+	if (!bInline)
+	{
+		// Handle the return value of the function being called
+		UProperty* FuncToCallReturnProperty = Statement.FunctionToCall->GetReturnProperty();
+		if (FuncToCallReturnProperty && ensure(Statement.LHS))
 		{
-			Result += TEXT("\t\t\t");
-			Result += SafeContextScope.GetAdditionalIndent();
-
-			// Handle the return value of the function being called
-			UProperty* FuncToCallReturnProperty = Statement.FunctionToCall->GetReturnProperty();
-			if (FuncToCallReturnProperty && ensure(Statement.LHS))
+			FString BeginCast;
+			if (auto ObjectProperty = Cast<UObjectProperty>(FuncToCallReturnProperty))
 			{
-				FString BeginCast;
-				if (auto ObjectProperty = Cast<UObjectProperty>(FuncToCallReturnProperty))
+				UClass* LClass = Statement.LHS ? Cast<UClass>(Statement.LHS->Type.PinSubCategoryObject.Get()) : nullptr;
+				if (LClass && LClass->IsChildOf(ObjectProperty->PropertyClass) && !ObjectProperty->PropertyClass->IsChildOf(LClass))
 				{
-					UClass* LClass = Statement.LHS ? Cast<UClass>(Statement.LHS->Type.PinSubCategoryObject.Get()) : nullptr;
-					if (LClass && LClass->IsChildOf(ObjectProperty->PropertyClass) && !ObjectProperty->PropertyClass->IsChildOf(LClass))
-					{
-						BeginCast = FString::Printf(TEXT("CastChecked<%s%s>("), LClass->GetPrefixCPP(), *LClass->GetName());
-						bCloseCast = true;
-					}
+					BeginCast = FString::Printf(TEXT("CastChecked<%s%s>("), LClass->GetPrefixCPP(), *LClass->GetName());
+					bCloseCast = true;
 				}
+			}
 
-				Result += FString::Printf(TEXT("%s = %s"), *TermToText(Statement.LHS), *BeginCast);
+			Result += FString::Printf(TEXT("%s = %s"), *TermToText(EmitterContext, Statement.LHS), *BeginCast);
 				
-			}
 		}
+	}
 
-		// Emit object to call the method on
-		if (bInterfaceCall)
+	// Emit object to call the method on
+	if (bInterfaceCall)
+	{
+		auto ContextInterfaceClass = CastChecked<UClass>(Statement.FunctionContext->Type.PinSubCategoryObject.Get());
+		ensure(ContextInterfaceClass->IsChildOf<UInterface>());
+		Result += FString::Printf(TEXT("I%s::Execute_%s(%s.GetObject(), ")
+			, *ContextInterfaceClass->GetName()
+			, *Statement.FunctionToCall->GetName()
+			, *TermToText(EmitterContext, Statement.FunctionContext, false));
+	}
+	else
+	{
+		if (bStaticCall)
 		{
-			auto ContextInterfaceClass = CastChecked<UClass>(Statement.FunctionContext->Type.PinSubCategoryObject.Get());
-			ensure(ContextInterfaceClass->IsChildOf<UInterface>());
-			Result += FString::Printf(TEXT("I%s::Execute_%s(%s.GetObject(), ")
-				, *ContextInterfaceClass->GetName()
-				, *Statement.FunctionToCall->GetName()
-				, *TermToText(Statement.FunctionContext, false));
+			const bool bIsCustomThunk = Statement.FunctionToCall->HasMetaData(TEXT("CustomStructureParam")) || Statement.FunctionToCall->HasMetaData(TEXT("ArrayParm"));
+			auto OwnerClass = Statement.FunctionToCall->GetOuterUClass();
+			Result += bIsCustomThunk ? TEXT("FCustomThunkTemplates::") : FString::Printf(TEXT("%s%s::"), OwnerClass->GetPrefixCPP(), *OwnerClass->GetName());
 		}
-		else
+		else if (bCallOnDifferentObject) //@TODO: Badness, could be a self reference wired to another instance!
 		{
-			if (bStaticCall)
-			{
-				const bool bIsCustomThunk = Statement.FunctionToCall->HasMetaData(TEXT("CustomStructureParam")) || Statement.FunctionToCall->HasMetaData(TEXT("ArrayParm"));
-				auto OwnerClass = Statement.FunctionToCall->GetOuterUClass();
-				Result += bIsCustomThunk ? TEXT("FCustomThunkTemplates::") : FString::Printf(TEXT("%s%s::"), OwnerClass->GetPrefixCPP(), *OwnerClass->GetName());
-			}
-			else if (bCallOnDifferentObject) //@TODO: Badness, could be a self reference wired to another instance!
-			{
-				Result += FString::Printf(TEXT("%s->"), *TermToText(Statement.FunctionContext, false));
-			}
-
-			if (Statement.bIsParentContext)
-			{
-				Result += TEXT("Super::");
-			}
-			Result += Statement.FunctionToCall->GetName();
-			if (Statement.bIsParentContext && FEmitHelper::ShoulsHandleAsNativeEvent(Statement.FunctionToCall))
-			{
-				ensure(!bCallOnDifferentObject);
-				Result += TEXT("_Implementation");
-			}
-
-			// Emit method parameter list
-			Result += TEXT("(");
+			Result += FString::Printf(TEXT("%s->"), *TermToText(EmitterContext, Statement.FunctionContext, false));
 		}
-		Result += EmitMethodInputParameterList(Statement);
-		Result += TEXT(")");
 
-		if (bCloseCast)
+		if (Statement.bIsParentContext)
 		{
-			Result += TEXT(", ECastCheckedType::NullAllowed)");
+			Result += TEXT("Super::");
 		}
-
-		if (!bInline)
+		Result += Statement.FunctionToCall->GetName();
+		if (Statement.bIsParentContext && FEmitHelper::ShoulsHandleAsNativeEvent(Statement.FunctionToCall))
 		{
-			Result += TEXT(";\n");
+			ensure(!bCallOnDifferentObject);
+			Result += TEXT("_Implementation");
 		}
+
+		// Emit method parameter list
+		Result += TEXT("(");
 	}
+	Result += EmitMethodInputParameterList(EmitterContext, Statement);
+	Result += TEXT(")");
+
+	if (bCloseCast)
+	{
+		Result += TEXT(", ECastCheckedType::NullAllowed)");
+	}
+
+	if (!bInline)
+	{
+		Result += TEXT(";");
+	}
+
 	return Result;
 }
 
-void FBlueprintCompilerCppBackend::EmitCallStatment(FKismetFunctionContext& FunctionContext, FBlueprintCompiledStatement& Statement)
+void FBlueprintCompilerCppBackend::EmitCallStatment(FEmitterLocalContext& EmitterContext, FKismetFunctionContext& FunctionContext, FBlueprintCompiledStatement& Statement)
 {
-	Emit(Body, *EmitCallStatmentInner(Statement, false));
+	const bool bCallOnDifferentObject = Statement.FunctionContext && (Statement.FunctionContext->Name != TEXT("self"));
+	const bool bStaticCall = Statement.FunctionToCall->HasAnyFunctionFlags(FUNC_Static);
+	const bool bUseSafeContext = bCallOnDifferentObject && !bStaticCall;
+	{
+		FSafeContextScopedEmmitter SafeContextScope(EmitterContext, bUseSafeContext ? Statement.FunctionContext : nullptr, *this);
+		FString Result = EmitCallStatmentInner(EmitterContext, Statement, false);
+		EmitterContext.AddLine(Result);
+	}
 }
 
-FString FBlueprintCompilerCppBackend::EmitSwitchValueStatmentInner(FBlueprintCompiledStatement& Statement)
+FString FBlueprintCompilerCppBackend::EmitSwitchValueStatmentInner(FEmitterLocalContext& EmitterContext, FBlueprintCompiledStatement& Statement)
 {
 	check(Statement.RHS.Num() >= 2);
 	const int32 TermsBeforeCases = 1;
@@ -394,8 +391,8 @@ FString FBlueprintCompilerCppBackend::EmitSwitchValueStatmentInner(FBlueprintCom
 	auto DefaultValueTerm = Statement.RHS.Last();
 
 	FString Result = FString::Printf(TEXT("TSwitchValue(%s, %s, %d")
-		, *TermToText(IndexTerm) //index
-		, *TermToText(DefaultValueTerm) // default
+		, *TermToText(EmitterContext, IndexTerm) //index
+		, *TermToText(EmitterContext, DefaultValueTerm) // default
 		, NumCases);
 	
 	const uint32 CppTemplateTypeFlags = EPropertyExportCPPFlags::CPPF_CustomTypeName 
@@ -415,8 +412,8 @@ FString FBlueprintCompilerCppBackend::EmitSwitchValueStatmentInner(FBlueprintCom
 		Result += FString::Printf(TEXT(", TSwitchPair<%s, %s>(%s, %s)")
 			, *IndexDeclaration
 			, *ValueDeclaration
-			, *TermToText(Statement.RHS[TermIndex])
-			, *TermToText(Statement.RHS[TermIndex + 1]));
+			, *TermToText(EmitterContext, Statement.RHS[TermIndex])
+			, *TermToText(EmitterContext, Statement.RHS[TermIndex + 1]));
 	}
 
 	Result += TEXT(")");
@@ -424,163 +421,139 @@ FString FBlueprintCompilerCppBackend::EmitSwitchValueStatmentInner(FBlueprintCom
 	return Result;
 }
 
-void FBlueprintCompilerCppBackend::EmitAssignmentStatment(FKismetFunctionContext& FunctionContext, FBlueprintCompiledStatement& Statement)
+void FBlueprintCompilerCppBackend::EmitAssignmentStatment(FEmitterLocalContext& EmitterContext, FKismetFunctionContext& FunctionContext, FBlueprintCompiledStatement& Statement)
 {
 	check(Statement.LHS && Statement.RHS[0]);
-	FString DestinationExpression = TermToText(Statement.LHS, false);
-	FString SourceExpression = TermToText(Statement.RHS[0]);
+	FString DestinationExpression = TermToText(EmitterContext, Statement.LHS, false);
+	FString SourceExpression = TermToText(EmitterContext, Statement.RHS[0]);
 
-	FSafeContextScopedEmmitter SafeContextScope(Body, Statement.LHS->Context, *this, TEXT("\t\t\t"));
+	FSafeContextScopedEmmitter SafeContextScope(EmitterContext, Statement.LHS->Context, *this);
 
 	FString BeginCast;
 	FString EndCast;
 	FEmitHelper::GenerateAssignmentCast(Statement.LHS->Type, Statement.RHS[0]->Type, BeginCast, EndCast);
-	// Emit the assignment statement
-	Emit(Body, TEXT("\t\t\t"));
-	Emit(Body, *SafeContextScope.GetAdditionalIndent());
-	Emit(Body, *FString::Printf(TEXT("%s = %s%s%s;\n"), *DestinationExpression, *BeginCast, *SourceExpression, *EndCast));
+	EmitterContext.AddLine(FString::Printf(TEXT("%s = %s%s%s;"), *DestinationExpression, *BeginCast, *SourceExpression, *EndCast));
 }
 
-void FBlueprintCompilerCppBackend::EmitCastObjToInterfaceStatement(FKismetFunctionContext& FunctionContext, FBlueprintCompiledStatement& Statement)
+void FBlueprintCompilerCppBackend::EmitCastObjToInterfaceStatement(FEmitterLocalContext& EmitterContext, FKismetFunctionContext& FunctionContext, FBlueprintCompiledStatement& Statement)
 {
-	FString InterfaceClass = TermToText(Statement.RHS[0]);
-	FString ObjectValue = TermToText(Statement.RHS[1]);
-	FString InterfaceValue = TermToText(Statement.LHS);
-
-	Emit(Body, *FString::Printf(TEXT("\t\t\tif ( IsValid(%s) && %s->GetClass()->ImplementsInterface(%s) )\n"), *ObjectValue, *ObjectValue, *InterfaceClass));
-	Emit(Body, *FString::Printf(TEXT("\t\t\t{\n")));
-	Emit(Body, *FString::Printf(TEXT("\t\t\t\t%s.SetObject(%s);\n"), *InterfaceValue, *ObjectValue));
-	Emit(Body, *FString::Printf(TEXT("\t\t\t\tvoid* IAddress = %s->GetInterfaceAddress(%s);\n"), *ObjectValue, *InterfaceClass));
-	Emit(Body, *FString::Printf(TEXT("\t\t\t\t%s.SetInterface(IAddress);\n"), *InterfaceValue));
-	Emit(Body, *FString::Printf(TEXT("\t\t\t}\n")));
-	Emit(Body, *FString::Printf(TEXT("\t\t\telse\n")));
-	Emit(Body, *FString::Printf(TEXT("\t\t\t{\n")));
-	Emit(Body, *FString::Printf(TEXT("\t\t\t\t%s.SetObject(NULL);\n"), *InterfaceValue));
-	Emit(Body, *FString::Printf(TEXT("\t\t\t}\n")));
+	FString InterfaceClass = TermToText(EmitterContext, Statement.RHS[0]);
+	FString ObjectValue = TermToText(EmitterContext, Statement.RHS[1]);
+	FString InterfaceValue = TermToText(EmitterContext, Statement.LHS);
+
+	EmitterContext.AddLine(FString::Printf(TEXT("if ( IsValid(%s) && %s->GetClass()->ImplementsInterface(%s) )"), *ObjectValue, *ObjectValue, *InterfaceClass));
+	EmitterContext.AddLine(FString::Printf(TEXT("{")));
+	EmitterContext.AddLine(FString::Printf(TEXT("\t%s.SetObject(%s);"), *InterfaceValue, *ObjectValue));
+	EmitterContext.AddLine(FString::Printf(TEXT("\tvoid* IAddress = %s->GetInterfaceAddress(%s);"), *ObjectValue, *InterfaceClass));
+	EmitterContext.AddLine(FString::Printf(TEXT("\t%s.SetInterface(IAddress);"), *InterfaceValue));
+	EmitterContext.AddLine(FString::Printf(TEXT("}")));
+	EmitterContext.AddLine(FString::Printf(TEXT("else")));
+	EmitterContext.AddLine(FString::Printf(TEXT("{")));
+	EmitterContext.AddLine(FString::Printf(TEXT("\t%s.SetObject(nullptr);"), *InterfaceValue));
+	EmitterContext.AddLine(FString::Printf(TEXT("}")));
 }
 
-void FBlueprintCompilerCppBackend::EmitCastBetweenInterfacesStatement(FKismetFunctionContext& FunctionContext, FBlueprintCompiledStatement& Statement)
+void FBlueprintCompilerCppBackend::EmitCastBetweenInterfacesStatement(FEmitterLocalContext& EmitterContext, FKismetFunctionContext& FunctionContext, FBlueprintCompiledStatement& Statement)
 {
-	FString ClassToCastTo = TermToText(Statement.RHS[0]);
-	FString InputInterface = TermToText(Statement.RHS[1]);
-	FString ResultInterface = TermToText(Statement.LHS);
+	FString ClassToCastTo = TermToText(EmitterContext, Statement.RHS[0]);
+	FString InputInterface = TermToText(EmitterContext, Statement.RHS[1]);
+	FString ResultInterface = TermToText(EmitterContext, Statement.LHS);
 
 	FString InputObject = FString::Printf(TEXT("%s.GetObjectRef()"), *InputInterface);
 
-	Emit(Body, *FString::Printf(TEXT("\t\t\tif ( %s && %s->GetClass()->IsChildOf(%s) )\n"), *InputObject, *InputObject, *ClassToCastTo));
-	Emit(Body, *FString::Printf(TEXT("\t\t\t{\n")));
-	Emit(Body, *FString::Printf(TEXT("\t\t\t\t%s.SetObject(%s);\n"), *ResultInterface, *InputObject));
-	Emit(Body, *FString::Printf(TEXT("\t\t\t\tvoid* IAddress = %s->GetInterfaceAddress(%s);\n"), *InputObject, *ClassToCastTo));
-	Emit(Body, *FString::Printf(TEXT("\t\t\t\t%s.SetInterface(IAddress);\n"), *ResultInterface));
-	Emit(Body, *FString::Printf(TEXT("\t\t\t}\n")));
-	Emit(Body, *FString::Printf(TEXT("\t\t\telse\n")));
-	Emit(Body, *FString::Printf(TEXT("\t\t\t{\n")));
-	Emit(Body, *FString::Printf(TEXT("\t\t\t\t%s.SetObject(NULL);\n"), *ResultInterface));
-	Emit(Body, *FString::Printf(TEXT("\t\t\t}\n")));
+	EmitterContext.AddLine(FString::Printf(TEXT("if ( %s && %s->GetClass()->IsChildOf(%s) )"), *InputObject, *InputObject, *ClassToCastTo));
+	EmitterContext.AddLine(FString::Printf(TEXT("{")));
+	EmitterContext.AddLine(FString::Printf(TEXT("\t%s.SetObject(%s);"), *ResultInterface, *InputObject));
+	EmitterContext.AddLine(FString::Printf(TEXT("\tvoid* IAddress = %s->GetInterfaceAddress(%s);"), *InputObject, *ClassToCastTo));
+	EmitterContext.AddLine(FString::Printf(TEXT("\t%s.SetInterface(IAddress);"), *ResultInterface));
+	EmitterContext.AddLine(FString::Printf(TEXT("}")));
+	EmitterContext.AddLine(FString::Printf(TEXT("else")));
+	EmitterContext.AddLine(FString::Printf(TEXT("{")));
+	EmitterContext.AddLine(FString::Printf(TEXT("\t%s.SetObject(nullptr);"), *ResultInterface));
+	EmitterContext.AddLine(FString::Printf(TEXT("}")));
 }
 
-void FBlueprintCompilerCppBackend::EmitCastInterfaceToObjStatement(FKismetFunctionContext& FunctionContext, FBlueprintCompiledStatement& Statement)
+void FBlueprintCompilerCppBackend::EmitCastInterfaceToObjStatement(FEmitterLocalContext& EmitterContext, FKismetFunctionContext& FunctionContext, FBlueprintCompiledStatement& Statement)
 {
-	FString ClassToCastTo = TermToText(Statement.RHS[0]);
-	FString InputInterface = TermToText(Statement.RHS[1]);
-	FString ResultObject = TermToText(Statement.LHS);
-
+	FString ClassToCastTo = TermToText(EmitterContext, Statement.RHS[0]);
+	FString InputInterface = TermToText(EmitterContext, Statement.RHS[1]);
+	FString ResultObject = TermToText(EmitterContext, Statement.LHS);
 	FString InputObject = FString::Printf(TEXT("%s.GetObjectRef()"), *InputInterface);
-
-	Emit(Body, *FString::Printf(TEXT("\t\t\t%s = Cast<%s>(%s);\n"),
-		*ResultObject, *ClassToCastTo, *InputObject));
+	EmitterContext.AddLine(FString::Printf(TEXT("%s = Cast<%s>(%s);"),*ResultObject, *ClassToCastTo, *InputObject));
 }
 
-void FBlueprintCompilerCppBackend::EmitDynamicCastStatement(FKismetFunctionContext& FunctionContext, FBlueprintCompiledStatement& Statement)
+void FBlueprintCompilerCppBackend::EmitDynamicCastStatement(FEmitterLocalContext& EmitterContext, FKismetFunctionContext& FunctionContext, FBlueprintCompiledStatement& Statement)
 {
-	FString TargetClass = TermToText(Statement.RHS[0]);
-	FString ObjectValue = TermToText(Statement.RHS[1]);
-	FString CastedValue = TermToText(Statement.LHS);
-
-	Emit(Body, *FString::Printf(TEXT("\t\t\t%s = Cast<%s>(%s);\n"),
-		*CastedValue, *TargetClass, *ObjectValue));
+	FString TargetClass = TermToText(EmitterContext, Statement.RHS[0]);
+	FString ObjectValue = TermToText(EmitterContext, Statement.RHS[1]);
+	FString CastedValue = TermToText(EmitterContext, Statement.LHS);
+	EmitterContext.AddLine(FString::Printf(TEXT("%s = Cast<%s>(%s);"),*CastedValue, *TargetClass, *ObjectValue));
 }
 
-void FBlueprintCompilerCppBackend::EmitMetaCastStatement(FKismetFunctionContext& FunctionContext, FBlueprintCompiledStatement& Statement)
+void FBlueprintCompilerCppBackend::EmitMetaCastStatement(FEmitterLocalContext& EmitterContext, FKismetFunctionContext& FunctionContext, FBlueprintCompiledStatement& Statement)
 {
-	FString DesiredClass = TermToText(Statement.RHS[0]);
-	FString SourceClass = TermToText(Statement.RHS[1]);
-	FString Destination = TermToText(Statement.LHS);
-
-	Emit(Body, *FString::Printf(TEXT("\t\t\t%s = DynamicMetaCast(%s, %s);\n"),
-		*Destination, *DesiredClass, *SourceClass));
+	FString DesiredClass = TermToText(EmitterContext, Statement.RHS[0]);
+	FString SourceClass = TermToText(EmitterContext, Statement.RHS[1]);
+	FString Destination = TermToText(EmitterContext, Statement.LHS);
+	EmitterContext.AddLine(FString::Printf(TEXT("%s = DynamicMetaCast(%s, %s);"),*Destination, *DesiredClass, *SourceClass));
 }
 
-void FBlueprintCompilerCppBackend::EmitObjectToBoolStatement(FKismetFunctionContext& FunctionContext, FBlueprintCompiledStatement& Statement)
+void FBlueprintCompilerCppBackend::EmitObjectToBoolStatement(FEmitterLocalContext& EmitterContext, FKismetFunctionContext& FunctionContext, FBlueprintCompiledStatement& Statement)
 {
-	FString ObjectTarget = TermToText(Statement.RHS[0]);
-	FString DestinationExpression = TermToText(Statement.LHS);
-
-	Emit(Body, *FString::Printf(TEXT("\t\t\t%s = (%s != NULL);\n"), *DestinationExpression, *ObjectTarget));
+	FString ObjectTarget = TermToText(EmitterContext, Statement.RHS[0]);
+	FString DestinationExpression = TermToText(EmitterContext, Statement.LHS);
+	EmitterContext.AddLine(FString::Printf(TEXT("%s = (%s != NULL);"), *DestinationExpression, *ObjectTarget));
 }
 
-void FBlueprintCompilerCppBackend::EmitAddMulticastDelegateStatement(FKismetFunctionContext& FunctionContext, FBlueprintCompiledStatement& Statement)
+void FBlueprintCompilerCppBackend::EmitAddMulticastDelegateStatement(FEmitterLocalContext& EmitterContext, FKismetFunctionContext& FunctionContext, FBlueprintCompiledStatement& Statement)
 {
 	check(Statement.LHS && Statement.LHS->AssociatedVarProperty);
-	FSafeContextScopedEmmitter SafeContextScope(Body, Statement.LHS->Context, *this, TEXT("\t\t\t"));
-
-	Emit(Body, TEXT("\t\t\t"));
-	Emit(Body, *SafeContextScope.GetAdditionalIndent());
+	FSafeContextScopedEmmitter SafeContextScope(EmitterContext, Statement.LHS->Context, *this);
 
-	const FString Delegate = TermToText(Statement.LHS, false);
-	const FString DelegateToAdd = TermToText(Statement.RHS[0]);
+	const FString Delegate = TermToText(EmitterContext, Statement.LHS, false);
+	const FString DelegateToAdd = TermToText(EmitterContext, Statement.RHS[0]);
 
-	Emit(Body, *FString::Printf(TEXT("%s.Add(%s);\n"), *Delegate, *DelegateToAdd));
+	EmitterContext.AddLine(FString::Printf(TEXT("%s.Add(%s);"), *Delegate, *DelegateToAdd));
 }
 
-void FBlueprintCompilerCppBackend::EmitRemoveMulticastDelegateStatement(FKismetFunctionContext& FunctionContext, FBlueprintCompiledStatement& Statement)
+void FBlueprintCompilerCppBackend::EmitRemoveMulticastDelegateStatement(FEmitterLocalContext& EmitterContext, FKismetFunctionContext& FunctionContext, FBlueprintCompiledStatement& Statement)
 {
 	check(Statement.LHS && Statement.LHS->AssociatedVarProperty);
-	FSafeContextScopedEmmitter SafeContextScope(Body, Statement.LHS->Context, *this, TEXT("\t\t\t"));
+	FSafeContextScopedEmmitter SafeContextScope(EmitterContext, Statement.LHS->Context, *this);
 
-	Emit(Body, TEXT("\t\t\t"));
-	Emit(Body, *SafeContextScope.GetAdditionalIndent());
+	const FString Delegate = TermToText(EmitterContext, Statement.LHS, false);
+	const FString DelegateToAdd = TermToText(EmitterContext, Statement.RHS[0]);
 
-	const FString Delegate = TermToText(Statement.LHS, false);
-	const FString DelegateToAdd = TermToText(Statement.RHS[0]);
-
-	Emit(Body, *FString::Printf(TEXT("%s.Remove(%s);\n"), *Delegate, *DelegateToAdd));
+	EmitterContext.AddLine(FString::Printf(TEXT("%s.Remove(%s);"), *Delegate, *DelegateToAdd));
 }
 
-void FBlueprintCompilerCppBackend::EmitBindDelegateStatement(FKismetFunctionContext& FunctionContext, FBlueprintCompiledStatement& Statement)
+void FBlueprintCompilerCppBackend::EmitBindDelegateStatement(FEmitterLocalContext& EmitterContext, FKismetFunctionContext& FunctionContext, FBlueprintCompiledStatement& Statement)
 {
 	check(2 == Statement.RHS.Num());
 	check(Statement.LHS);
-	FSafeContextScopedEmmitter SafeContextScope(Body, Statement.LHS->Context, *this, TEXT("\t\t\t"));
-
-	Emit(Body, TEXT("\t\t\t"));
-	Emit(Body, *SafeContextScope.GetAdditionalIndent());
-
+	FSafeContextScopedEmmitter SafeContextScope(EmitterContext, Statement.LHS->Context, *this);
 
-	const FString Delegate = TermToText(Statement.LHS, false);
-	const FString NameTerm = TermToText(Statement.RHS[0]);
-	const FString ObjectTerm = TermToText(Statement.RHS[1]);
+	const FString Delegate = TermToText(EmitterContext, Statement.LHS, false);
+	const FString NameTerm = TermToText(EmitterContext, Statement.RHS[0]);
+	const FString ObjectTerm = TermToText(EmitterContext, Statement.RHS[1]);
 
-	Emit(Body, *FString::Printf(TEXT("%s.BindUFunction(%s,%s);\n"), *Delegate, *ObjectTerm, *NameTerm));
+	EmitterContext.AddLine(FString::Printf(TEXT("%s.BindUFunction(%s,%s);"), *Delegate, *ObjectTerm, *NameTerm));
 }
 
-void FBlueprintCompilerCppBackend::EmitClearMulticastDelegateStatement(FKismetFunctionContext& FunctionContext, FBlueprintCompiledStatement& Statement)
+void FBlueprintCompilerCppBackend::EmitClearMulticastDelegateStatement(FEmitterLocalContext& EmitterContext, FKismetFunctionContext& FunctionContext, FBlueprintCompiledStatement& Statement)
 {
 	check(Statement.LHS);
-	FSafeContextScopedEmmitter SafeContextScope(Body, Statement.LHS->Context, *this, TEXT("\t\t\t"));
+	FSafeContextScopedEmmitter SafeContextScope(EmitterContext, Statement.LHS->Context, *this);
 
-	Emit(Body, TEXT("\t\t\t"));
-	Emit(Body, *SafeContextScope.GetAdditionalIndent());
+	const FString Delegate = TermToText(EmitterContext, Statement.LHS, false);
 
-	const FString Delegate = TermToText(Statement.LHS, false);
-
-	Emit(Body, *FString::Printf(TEXT("%s.Clear();\n"), *Delegate));
+	EmitterContext.AddLine(FString::Printf(TEXT("%s.Clear();"), *Delegate));
 }
 
-void FBlueprintCompilerCppBackend::EmitCreateArrayStatement(FKismetFunctionContext& FunctionContext, FBlueprintCompiledStatement& Statement)
+void FBlueprintCompilerCppBackend::EmitCreateArrayStatement(FEmitterLocalContext& EmitterContext, FKismetFunctionContext& FunctionContext, FBlueprintCompiledStatement& Statement)
 {
 	FBPTerminal* ArrayTerm = Statement.LHS;
-	const FString Array = TermToText(ArrayTerm);
+	const FString Array = TermToText(EmitterContext, ArrayTerm);
 
 	UArrayProperty* ArrayProperty = CastChecked<UArrayProperty>(ArrayTerm->AssociatedVarProperty);
 	UProperty* InnerProperty = ArrayProperty->Inner;
@@ -588,68 +561,69 @@ void FBlueprintCompilerCppBackend::EmitCreateArrayStatement(FKismetFunctionConte
 	for (int32 i = 0; i < Statement.RHS.Num(); ++i)
 	{
 		FBPTerminal* CurrentTerminal = Statement.RHS[i];
-		Emit(Body,*FString::Printf(TEXT("\t\t\t%s[%d] = %s;"), *Array, i, *TermToText(CurrentTerminal)));
+		EmitterContext.AddLine(FString::Printf(TEXT("%s[%d] = %s;"), *Array, i, *TermToText(EmitterContext, CurrentTerminal)));
 	}
 }
 
-void FBlueprintCompilerCppBackend::EmitGotoStatement(FKismetFunctionContext& FunctionContext, FBlueprintCompiledStatement& Statement)
+void FBlueprintCompilerCppBackend::EmitGotoStatement(FEmitterLocalContext& EmitterContext, FKismetFunctionContext& FunctionContext, FBlueprintCompiledStatement& Statement)
 {
 	if (Statement.Type == KCST_ComputedGoto)
 	{
 		FString NextStateExpression;
-		NextStateExpression = TermToText(Statement.LHS);
+		NextStateExpression = TermToText(EmitterContext, Statement.LHS);
 
-		Emit(Body, *FString::Printf(TEXT("\t\t\tCurrentState = %s;\n"), *NextStateExpression));
-		Emit(Body, *FString::Printf(TEXT("\t\t\tbreak;\n")));
+		EmitterContext.AddLine(FString::Printf(TEXT("CurrentState = %s;"), *NextStateExpression));
+		EmitterContext.AddLine(FString::Printf(TEXT("break;\n")));
 	}
 	else if ((Statement.Type == KCST_GotoIfNot) || (Statement.Type == KCST_EndOfThreadIfNot) || (Statement.Type == KCST_GotoReturnIfNot))
 	{
 		FString ConditionExpression;
-		ConditionExpression = TermToText(Statement.LHS);
-
-		Emit(Body, *FString::Printf(TEXT("\t\t\tif (!%s)\n"), *ConditionExpression));
-		Emit(Body, *FString::Printf(TEXT("\t\t\t{\n")));
+		ConditionExpression = TermToText(EmitterContext, Statement.LHS);
 
+		EmitterContext.AddLine(FString::Printf(TEXT("if (!%s)"), *ConditionExpression));
+		EmitterContext.AddLine(FString::Printf(TEXT("{")));
+		EmitterContext.IncreaseIndent();
 		if (Statement.Type == KCST_EndOfThreadIfNot)
 		{
 			ensure(FunctionContext.bUseFlowStack);
-			Emit(Body, TEXT("\t\t\t\tCurrentState = (StateStack.Num() > 0) ? StateStack.Pop(/*bAllowShrinking=*/ false) : -1;\n"));
+			EmitterContext.AddLine(TEXT("CurrentState = (StateStack.Num() > 0) ? StateStack.Pop(/*bAllowShrinking=*/ false) : -1;"));
 		}
 		else if (Statement.Type == KCST_GotoReturnIfNot)
 		{
-			Emit(Body, TEXT("\t\t\t\tCurrentState = -1;\n"));
+			EmitterContext.AddLine(TEXT("CurrentState = -1;"));
 		}
 		else
 		{
-			Emit(Body, *FString::Printf(TEXT("\t\t\t\tCurrentState = %d;\n"), StatementToStateIndex(FunctionContext, Statement.TargetLabel)));
+			EmitterContext.AddLine(FString::Printf(TEXT("CurrentState = %d;"), StatementToStateIndex(FunctionContext, Statement.TargetLabel)));
 		}
 
-		Emit(Body, *FString::Printf(TEXT("\t\t\t\tbreak;\n")));
-		Emit(Body, *FString::Printf(TEXT("\t\t\t}\n")));
+		EmitterContext.AddLine(FString::Printf(TEXT("break;")));
+		EmitterContext.DecreaseIndent();
+		EmitterContext.AddLine(FString::Printf(TEXT("}")));
 	}
 	else if (Statement.Type == KCST_GotoReturn)
 	{
-		Emit(Body, TEXT("\t\t\tCurrentState = -1;\n"));
-		Emit(Body, *FString::Printf(TEXT("\t\t\tbreak;\n")));
+		EmitterContext.AddLine(TEXT("CurrentState = -1;"));
+		EmitterContext.AddLine(FString::Printf(TEXT("break;")));
 	}
 	else
 	{
-		Emit(Body, *FString::Printf(TEXT("\t\t\tCurrentState = %d;\n"), StatementToStateIndex(FunctionContext, Statement.TargetLabel)));
-		Emit(Body, *FString::Printf(TEXT("\t\t\tbreak;\n")));
+		EmitterContext.AddLine(FString::Printf(TEXT("CurrentState = %d;"), StatementToStateIndex(FunctionContext, Statement.TargetLabel)));
+		EmitterContext.AddLine(FString::Printf(TEXT("break;")));
 	}
 }
 
-void FBlueprintCompilerCppBackend::EmitPushStateStatement(FKismetFunctionContext& FunctionContext, FBlueprintCompiledStatement& Statement)
+void FBlueprintCompilerCppBackend::EmitPushStateStatement(FEmitterLocalContext& EmitterContext, FKismetFunctionContext& FunctionContext, FBlueprintCompiledStatement& Statement)
 {
 	ensure(FunctionContext.bUseFlowStack);
-	Emit(Body, *FString::Printf(TEXT("\t\t\tStateStack.Push(%d);\n"), StatementToStateIndex(FunctionContext, Statement.TargetLabel)));
+	EmitterContext.AddLine(FString::Printf(TEXT("StateStack.Push(%d);"), StatementToStateIndex(FunctionContext, Statement.TargetLabel)));
 }
 
-void FBlueprintCompilerCppBackend::EmitEndOfThreadStatement(FKismetFunctionContext& FunctionContext, const FString& ReturnValueString)
+void FBlueprintCompilerCppBackend::EmitEndOfThreadStatement(FEmitterLocalContext& EmitterContext, FKismetFunctionContext& FunctionContext, const FString& ReturnValueString)
 {
 	ensure(FunctionContext.bUseFlowStack);
-	Emit(Body, TEXT("\t\t\tCurrentState = (StateStack.Num() > 0) ? StateStack.Pop(/*bAllowShrinking=*/ false) : -1;\n"));
-	Emit(Body, TEXT("\t\t\tbreak;\n"));
+	EmitterContext.AddLine(TEXT("CurrentState = (StateStack.Num() > 0) ? StateStack.Pop(/*bAllowShrinking=*/ false) : -1;"));
+	EmitterContext.AddLine(TEXT("break;"));
 }
 
 void FBlueprintCompilerCppBackend::EmitReturnStatement(FKismetFunctionContext& FunctionContext, const FString& ReturnValueString)
@@ -686,11 +660,13 @@ void FBlueprintCompilerCppBackend::DeclareStateSwitch(FKismetFunctionContext& Fu
 	Emit(Body, TEXT("\t\tswitch( CurrentState )\n"));
 	Emit(Body, TEXT("\t\t{\n"));
 	Emit(Body, TEXT("\t\tcase 0:\n"));
+	Emit(Body, TEXT("\t\t\t{\n"));
 }
 
 void FBlueprintCompilerCppBackend::CloseStateSwitch(FKismetFunctionContext& FunctionContext)
 {
 	// Default error-catching case 
+	Emit(Body, TEXT("\t\t\t}\n"));
 	Emit(Body, TEXT("\t\tdefault:\n"));
 	if (FunctionContext.bUseFlowStack)
 	{
@@ -893,6 +869,16 @@ void FBlueprintCompilerCppBackend::ConstructFunction(FKismetFunctionContext& Fun
 			}
 		}
 
+		FEmitterLocalContext EmitterContext;
+		EmitterContext.IncreaseIndent();
+		if (bUseSwitchState)
+		{
+			EmitterContext.IncreaseIndent();
+			EmitterContext.IncreaseIndent();
+			EmitterContext.IncreaseIndent();
+		}
+		EmitterContext.SetCurrentlyGeneratedClass(FunctionContext.NewClass);
+
 		// Emit code in the order specified by the linear execution list (the first node is always the entry point for the function)
 		for (int32 NodeIndex = 0; NodeIndex < FunctionContext.LinearExecutionList.Num(); ++NodeIndex)
 		{
@@ -907,78 +893,82 @@ void FBlueprintCompilerCppBackend::ConstructFunction(FKismetFunctionContext& Fun
 
 					if (Statement.bIsJumpTarget && bUseSwitchState)
 					{
+						EmitterContext.DecreaseIndent();
+						EmitterContext.AddLine(TEXT("}"));
 						const int32 StateNum = StatementToStateIndex(FunctionContext, &Statement);
-						Emit(Body, *FString::Printf(TEXT("\n\t\tcase %d:\n"), StateNum));
+						EmitterContext.DecreaseIndent();
+						EmitterContext.AddLine(FString::Printf(TEXT("case %d:"), StateNum));
+						EmitterContext.IncreaseIndent();
+						EmitterContext.AddLine(TEXT("{"));
+						EmitterContext.IncreaseIndent();
 					}
 
 					switch (Statement.Type)
 					{
 					case KCST_Nop:
-						Emit(Body, TEXT("\t\t\t//No operation.\n"));
+						EmitterContext.AddLine(TEXT("//No operation."));
 						break;
 					case KCST_CallFunction:
-						EmitCallStatment(FunctionContext, Statement);
+						EmitCallStatment(EmitterContext, FunctionContext, Statement);
 						break;
 					case KCST_Assignment:
-						EmitAssignmentStatment(FunctionContext, Statement);
+						EmitAssignmentStatment(EmitterContext, FunctionContext, Statement);
 						break;
 					case KCST_CompileError:
 						UE_LOG(LogK2Compiler, Error, TEXT("C++ backend encountered KCST_CompileError"));
-						Emit(Body, TEXT("\t\t\tstatic_assert(false); // KCST_CompileError"));
+						EmitterContext.AddLine(TEXT("static_assert(false); // KCST_CompileError"));
 						break;
 					case KCST_PushState:
-						EmitPushStateStatement(FunctionContext, Statement);
+						EmitPushStateStatement(EmitterContext, FunctionContext, Statement);
 						break;
 					case KCST_Return:
-						Emit(Body, TEXT("\t\t\t// Return statement.\n"));
+						EmitterContext.AddLine(TEXT("// Return statement."));
 						break;
 					case KCST_EndOfThread:
-						EmitEndOfThreadStatement(FunctionContext, ReturnValueString);
+						EmitEndOfThreadStatement(EmitterContext, FunctionContext, ReturnValueString);
 						break;
 					case KCST_Comment:
-						Emit(Body, *FString::Printf(TEXT("\t\t\t// %s\n"), *Statement.Comment));
+						EmitterContext.AddLine(FString::Printf(TEXT("// %s"), *Statement.Comment));
 						break;
 					case KCST_DebugSite:
-						Emit(Body, TEXT("\t\t\t// Debug site.\n"));
 						break;
 					case KCST_CastObjToInterface:
-						EmitCastObjToInterfaceStatement(FunctionContext, Statement);
+						EmitCastObjToInterfaceStatement(EmitterContext, FunctionContext, Statement);
 						break;
 					case KCST_DynamicCast:
-						EmitDynamicCastStatement(FunctionContext, Statement);
+						EmitDynamicCastStatement(EmitterContext, FunctionContext, Statement);
 						break;
 					case KCST_ObjectToBool:
-						EmitObjectToBoolStatement(FunctionContext, Statement);
+						EmitObjectToBoolStatement(EmitterContext, FunctionContext, Statement);
 						break;
 					case KCST_AddMulticastDelegate:
-						EmitAddMulticastDelegateStatement(FunctionContext, Statement);
+						EmitAddMulticastDelegateStatement(EmitterContext, FunctionContext, Statement);
 						break;
 					case KCST_ClearMulticastDelegate:
-						EmitClearMulticastDelegateStatement(FunctionContext, Statement);
+						EmitClearMulticastDelegateStatement(EmitterContext, FunctionContext, Statement);
 						break;
 					case KCST_WireTraceSite:
-						Emit(Body, TEXT("\t\t\t// Wire debug site.\n"));
 						break;
 					case KCST_BindDelegate:
-						EmitBindDelegateStatement(FunctionContext, Statement);
+						EmitBindDelegateStatement(EmitterContext, FunctionContext, Statement);
 						break;
 					case KCST_RemoveMulticastDelegate:
-						EmitRemoveMulticastDelegateStatement(FunctionContext, Statement);
+						EmitRemoveMulticastDelegateStatement(EmitterContext, FunctionContext, Statement);
 						break;
 					case KCST_CallDelegate:
-						EmitCallDelegateStatment(FunctionContext, Statement);
+						EmitCallDelegateStatment(EmitterContext, FunctionContext, Statement);
 						break;
 					case KCST_CreateArray:
-						EmitCreateArrayStatement(FunctionContext, Statement);
+						EmitCreateArrayStatement(EmitterContext, FunctionContext, Statement);
 						break;
 					case KCST_CrossInterfaceCast:
-						EmitCastBetweenInterfacesStatement(FunctionContext, Statement);
+						EmitCastBetweenInterfacesStatement(EmitterContext, FunctionContext, Statement);
 						break;
 					case KCST_MetaCast:
-						EmitMetaCastStatement(FunctionContext, Statement);
+						EmitMetaCastStatement(EmitterContext, FunctionContext, Statement);
 						break;
 					case KCST_CastInterfaceToObj:
-						EmitCastInterfaceToObjStatement(FunctionContext, Statement);
+						EmitCastInterfaceToObjStatement(EmitterContext, FunctionContext, Statement);
 						break;
 					case KCST_ComputedGoto:
 					case KCST_UnconditionalGoto:
@@ -986,13 +976,13 @@ void FBlueprintCompilerCppBackend::ConstructFunction(FKismetFunctionContext& Fun
 					case KCST_EndOfThreadIfNot:
 					case KCST_GotoReturn:
 					case KCST_GotoReturnIfNot:
-						EmitGotoStatement(FunctionContext, Statement);
+						EmitGotoStatement(EmitterContext, FunctionContext, Statement);
 						break;
 					case KCST_SwitchValue:
 						// Switch Value should be always an "inline" statement, so there is no point to handle it here
 						// case: KCST_AssignmentOnPersistentFrame
 					default:
-						Emit(Body, TEXT("\t// Warning: Ignoring unsupported statement\n"));
+						EmitterContext.AddLine(TEXT("// Warning: Ignoring unsupported statement\n"));
 						UE_LOG(LogK2Compiler, Error, TEXT("C++ backend encountered unsupported statement type %d"), (int32)Statement.Type);
 						break;
 					};
@@ -1000,6 +990,7 @@ void FBlueprintCompilerCppBackend::ConstructFunction(FKismetFunctionContext& Fun
 			}
 		}
 
+		Emit(Body, *EmitterContext.GetResult());
 		if (bUseSwitchState)
 		{
 			CloseStateSwitch(FunctionContext);
diff --git a/Engine/Source/Developer/BlueprintCompilerCppBackend/Private/BlueprintCompilerCppBackend.h b/Engine/Source/Developer/BlueprintCompilerCppBackend/Private/BlueprintCompilerCppBackend.h
--- a/Engine/Source/Developer/BlueprintCompilerCppBackend/Private/BlueprintCompilerCppBackend.h
+++ b/Engine/Source/Developer/BlueprintCompilerCppBackend/Private/BlueprintCompilerCppBackend.h
@@ -3,6 +3,8 @@
 
 #include "IBlueprintCompilerCppBackendModule.h"
 
+struct FEmitterLocalContext;
+
 class FBlueprintCompilerCppBackend : public IBlueprintCompilerCppBackend
 {
 protected:
@@ -55,28 +57,29 @@ public:
 
 	void EmitStructProperties(FStringOutputDevice& Target, UStruct* SourceClass);
 
-	void EmitCallStatment(FKismetFunctionContext& FunctionContext, FBlueprintCompiledStatement& Statement);
-	void EmitCallDelegateStatment(FKismetFunctionContext& FunctionContext, FBlueprintCompiledStatement& Statement);
-	void EmitAssignmentStatment(FKismetFunctionContext& FunctionContext, FBlueprintCompiledStatement& Statement);
-	void EmitCastObjToInterfaceStatement(FKismetFunctionContext& FunctionContext, FBlueprintCompiledStatement& Statement);
-	void EmitCastBetweenInterfacesStatement(FKismetFunctionContext& FunctionContext, FBlueprintCompiledStatement& Statement);
-	void EmitCastInterfaceToObjStatement(FKismetFunctionContext& FunctionContext, FBlueprintCompiledStatement& Statement);
-	void EmitDynamicCastStatement(FKismetFunctionContext& FunctionContext, FBlueprintCompiledStatement& Statement);
-	void EmitObjectToBoolStatement(FKismetFunctionContext& FunctionContext, FBlueprintCompiledStatement& Statement);
-	void EmitAddMulticastDelegateStatement(FKismetFunctionContext& FunctionContext, FBlueprintCompiledStatement& Statement);
-	void EmitBindDelegateStatement(FKismetFunctionContext& FunctionContext, FBlueprintCompiledStatement& Statement);
-	void EmitClearMulticastDelegateStatement(FKismetFunctionContext& FunctionContext, FBlueprintCompiledStatement& Statement);
-	void EmitCreateArrayStatement(FKismetFunctionContext& FunctionContext, FBlueprintCompiledStatement& Statement);
-	void EmitGotoStatement(FKismetFunctionContext& FunctionContext, FBlueprintCompiledStatement& Statement);
-	void EmitPushStateStatement(FKismetFunctionContext& FunctionContext, FBlueprintCompiledStatement& Statement);
-	void EmitEndOfThreadStatement(FKismetFunctionContext& FunctionContext, const FString& ReturnValueString);
+	void EmitCallStatment(FEmitterLocalContext& EmitterContext, FKismetFunctionContext& FunctionContext, FBlueprintCompiledStatement& Statement);
+	void EmitCallDelegateStatment(FEmitterLocalContext& EmitterContext, FKismetFunctionContext& FunctionContext, FBlueprintCompiledStatement& Statement);
+	void EmitAssignmentStatment(FEmitterLocalContext& EmitterContext, FKismetFunctionContext& FunctionContext, FBlueprintCompiledStatement& Statement);
+	void EmitCastObjToInterfaceStatement(FEmitterLocalContext& EmitterContext, FKismetFunctionContext& FunctionContext, FBlueprintCompiledStatement& Statement);
+	void EmitCastBetweenInterfacesStatement(FEmitterLocalContext& EmitterContext, FKismetFunctionContext& FunctionContext, FBlueprintCompiledStatement& Statement);
+	void EmitCastInterfaceToObjStatement(FEmitterLocalContext& EmitterContext, FKismetFunctionContext& FunctionContext, FBlueprintCompiledStatement& Statement);
+	void EmitDynamicCastStatement(FEmitterLocalContext& EmitterContext, FKismetFunctionContext& FunctionContext, FBlueprintCompiledStatement& Statement);
+	void EmitObjectToBoolStatement(FEmitterLocalContext& EmitterContext, FKismetFunctionContext& FunctionContext, FBlueprintCompiledStatement& Statement);
+	void EmitAddMulticastDelegateStatement(FEmitterLocalContext& EmitterContext, FKismetFunctionContext& FunctionContext, FBlueprintCompiledStatement& Statement);
+	void EmitBindDelegateStatement(FEmitterLocalContext& EmitterContext, FKismetFunctionContext& FunctionContext, FBlueprintCompiledStatement& Statement);
+	void EmitClearMulticastDelegateStatement(FEmitterLocalContext& EmitterContext, FKismetFunctionContext& FunctionContext, FBlueprintCompiledStatement& Statement);
+	void EmitCreateArrayStatement(FEmitterLocalContext& EmitterContext, FKismetFunctionContext& FunctionContext, FBlueprintCompiledStatement& Statement);
+	void EmitGotoStatement(FEmitterLocalContext& EmitterContext, FKismetFunctionContext& FunctionContext, FBlueprintCompiledStatement& Statement);
+	void EmitPushStateStatement(FEmitterLocalContext& EmitterContext, FKismetFunctionContext& FunctionContext, FBlueprintCompiledStatement& Statement);
+	void EmitEndOfThreadStatement(FEmitterLocalContext& EmitterContext, FKismetFunctionContext& FunctionContext, const FString& ReturnValueString);
+	void EmitRemoveMulticastDelegateStatement(FEmitterLocalContext& EmitterContext, FKismetFunctionContext& FunctionContext, FBlueprintCompiledStatement& Statement);
+	void EmitMetaCastStatement(FEmitterLocalContext& EmitterContext, FKismetFunctionContext& FunctionContext, FBlueprintCompiledStatement& Statement);
+
 	void EmitReturnStatement(FKismetFunctionContext& FunctionContext, const FString& ReturnValueString);
-	void EmitRemoveMulticastDelegateStatement(FKismetFunctionContext& FunctionContext, FBlueprintCompiledStatement& Statement);
-	void EmitMetaCastStatement(FKismetFunctionContext& FunctionContext, FBlueprintCompiledStatement& Statement);
 
 	/** Emits local variable declarations for a function */
 	void DeclareLocalVariables(FKismetFunctionContext& FunctionContext, TArray<UProperty*>& LocalVariables);
-
+	
 	/** Emits the internal execution flow state declaration for a function */
 	void DeclareStateSwitch(FKismetFunctionContext& FunctionContext);
 	void CloseStateSwitch(FKismetFunctionContext& FunctionContext);
@@ -84,22 +87,22 @@ public:
 	/** Builds both the header declaration and body implementation of a function */
 	void ConstructFunction(FKismetFunctionContext& FunctionContext, bool bGenerateStubOnly);
 
-	FString TermToText(const FBPTerminal* Term, bool bUseSafeContext = true);
+	FString TermToText(FEmitterLocalContext& EmitterContext, const FBPTerminal* Term, bool bUseSafeContext = true);
 
 protected:
 	void EmitFileBeginning(const FString& CleanName, UStruct* SourceStruct);
 
-	FString LatentFunctionInfoTermToText(FBPTerminal* Term, FBlueprintCompiledStatement* TargetLabel);
-
 	int32 StatementToStateIndex(FKismetFunctionContext& FunctionContext, FBlueprintCompiledStatement* Statement)
 	{
 		int32 Index = FunctionIndexMap.FindChecked(&FunctionContext);
 		return StateMapPerFunction[Index].StatementToStateIndex(Statement);
 	}
 
-	FString EmitMethodInputParameterList(FBlueprintCompiledStatement& Statement);
+	FString LatentFunctionInfoTermToText(FEmitterLocalContext& EmitterContext, FBPTerminal* Term, FBlueprintCompiledStatement* TargetLabel);
+	
+	FString EmitMethodInputParameterList(FEmitterLocalContext& EmitterContext, FBlueprintCompiledStatement& Statement);
 
-	FString EmitSwitchValueStatmentInner(FBlueprintCompiledStatement& Statement);
+	FString EmitSwitchValueStatmentInner(FEmitterLocalContext& EmitterContext, FBlueprintCompiledStatement& Statement);
 
-	FString EmitCallStatmentInner(FBlueprintCompiledStatement& Statement, bool bInline);
+	FString EmitCallStatmentInner(FEmitterLocalContext& EmitterContext, FBlueprintCompiledStatement& Statement, bool bInline);
 };
\ No newline at end of file
diff --git a/Engine/Source/Developer/BlueprintCompilerCppBackend/Private/BlueprintCompilerCppBackendUtils.cpp b/Engine/Source/Developer/BlueprintCompilerCppBackend/Private/BlueprintCompilerCppBackendUtils.cpp
--- a/Engine/Source/Developer/BlueprintCompilerCppBackend/Private/BlueprintCompilerCppBackendUtils.cpp
+++ b/Engine/Source/Developer/BlueprintCompilerCppBackend/Private/BlueprintCompilerCppBackendUtils.cpp
@@ -8,17 +8,12 @@
 #include "Engine/InheritableComponentHandler.h"
 #include "Engine/DynamicBlueprintBinding.h"
 
-FString FSafeContextScopedEmmitter::GetAdditionalIndent() const
-{
-	return bSafeContextUsed ? TEXT("\t") : FString();
-}
-
 bool FSafeContextScopedEmmitter::IsSafeContextUsed() const
 {
 	return bSafeContextUsed;
 }
 
-FString FSafeContextScopedEmmitter::ValidationChain(const FBPTerminal* Term, FBlueprintCompilerCppBackend& CppBackend)
+FString FSafeContextScopedEmmitter::ValidationChain(FEmitterLocalContext& EmitterContext, const FBPTerminal* Term, FBlueprintCompilerCppBackend& CppBackend)
 {
 	TArray<FString> SafetyConditions;
 	for (; Term; Term = Term->Context)
@@ -26,7 +21,7 @@ FString FSafeContextScopedEmmitter::ValidationChain(const FBPTerminal* Term, FBl
 		if (!Term->IsStructContextType() && (Term->Type.PinSubCategory != TEXT("self")))
 		{
 			ensure(!Term->bIsLiteral);
-			SafetyConditions.Add(CppBackend.TermToText(Term, false));
+			SafetyConditions.Add(CppBackend.TermToText(EmitterContext, Term, false));
 		}
 	}
 
@@ -45,20 +40,18 @@ FString FSafeContextScopedEmmitter::ValidationChain(const FBPTerminal* Term, FBl
 	return Result;
 }
 
-FSafeContextScopedEmmitter::FSafeContextScopedEmmitter(FString& InBody, const FBPTerminal* Term, FBlueprintCompilerCppBackend& CppBackend, const TCHAR* InCurrentIndent)
-	: Body(InBody), bSafeContextUsed(false), CurrentIndent(InCurrentIndent)
+FSafeContextScopedEmmitter::FSafeContextScopedEmmitter(FEmitterLocalContext& InEmitterContext, const FBPTerminal* Term, FBlueprintCompilerCppBackend& CppBackend)
+	: EmitterContext(InEmitterContext)
+	, bSafeContextUsed(false)
 {
-	const FString Conditions = ValidationChain(Term, CppBackend);
+	const FString Conditions = ValidationChain(EmitterContext, Term, CppBackend);
 
 	if (!Conditions.IsEmpty())
 	{
 		bSafeContextUsed = true;
-		Body += CurrentIndent;
-		Body += TEXT("if (");
-		Body += Conditions;
-		Body += TEXT(")\n");
-		Body += CurrentIndent;
-		Body += TEXT("{\n");
+		EmitterContext.AddLine(FString::Printf(TEXT("if(%s)"), *Conditions));
+		EmitterContext.AddLine(TEXT("{"));
+		EmitterContext.IncreaseIndent();
 	}
 }
 
@@ -66,8 +59,8 @@ FSafeContextScopedEmmitter::~FSafeContextScopedEmmitter()
 {
 	if (bSafeContextUsed)
 	{
-		Body += CurrentIndent;
-		Body += TEXT("}\n");
+		EmitterContext.DecreaseIndent();
+		EmitterContext.AddLine(TEXT("}"));
 	}
 }
 
@@ -438,7 +431,7 @@ FString FEmitHelper::GatherNativeHeadersToInclude(UField* SourceItem, const TArr
 	return Result;
 }
 
-FString FEmitHelper::LiteralTerm(const FEdGraphPinType& Type, const FString& CustomValue, UObject* LiteralObject)
+FString FEmitHelper::LiteralTerm(FEmitterLocalContext& EmitterContext, const FEdGraphPinType& Type, const FString& CustomValue, UObject* LiteralObject)
 {
 	auto Schema = GetDefault<UEdGraphSchema_K2>();
 
@@ -534,16 +527,30 @@ FString FEmitHelper::LiteralTerm(const FEdGraphPinType& Type, const FString& Cus
 		{
 			//@todo:  This needs to be more robust, since import text isn't really proper for struct construction.
 			const bool bEmptyCustomValue = CustomValue.IsEmpty() || (CustomValue == TEXT("()"));
-			ensure(bEmptyCustomValue);
-			
 			const FString StructName = FString(TEXT("F")) + StructType->GetName();
-			const FString ConstructBody = bEmptyCustomValue
-				? (StructType->IsA<UUserDefinedStruct>() ? TEXT("::GetDefaultValue()") : TEXT("{}"))
-				: CustomValue;
+			if (bEmptyCustomValue)
+			{
+				return StructName + (StructType->IsA<UUserDefinedStruct>() ? TEXT("::GetDefaultValue()") : TEXT("{}"));
+			}
 
-			// Import.. export..
+			const FString LocalStructNativeName = EmitterContext.GenerateUniqueLocalName();
+			EmitterContext.AddLine(FString::Printf(TEXT("auto %s = %s{};"), *LocalStructNativeName, *StructName)); // TODO: ?? should "::GetDefaultValue()" be called here?
+			{
+				FStructOnScope StructOnScope(StructType);
 
-			return StructName + ConstructBody;
+				FStringOutputDevice ImportError;
+				const auto EndOfParsedBuff = UStructProperty::ImportText_Static(StructType, TEXT("FEmitHelper::LiteralTerm"), *CustomValue, StructOnScope.GetStructMemory(), 0, nullptr, &ImportError);
+				if (!EndOfParsedBuff || !ImportError.IsEmpty())
+				{
+					UE_LOG(LogK2Compiler, Error, TEXT("FEmitHelper::LiteralTerm cannot parse struct \"%s\" error: %s"), *CustomValue, *ImportError);
+				}
+
+				for (auto LocalProperty : TFieldRange<const UProperty>(StructType))
+				{
+					FEmitDefaultValueHelper::OuterGenerate(EmitterContext, LocalProperty, LocalStructNativeName, StructOnScope.GetStructMemory(), nullptr, FEmitDefaultValueHelper::EPropertyAccessOperator::Dot);
+				}
+			}
+			return LocalStructNativeName;
 		}
 	}
 	else if (Type.PinSubCategory == UEdGraphSchema_K2::PSC_Self)
@@ -595,9 +602,9 @@ FString FEmitHelper::LiteralTerm(const FEdGraphPinType& Type, const FString& Cus
 	return CustomValue;
 }
 
-FString FEmitHelper::DefaultValue(const FEdGraphPinType& Type)
+FString FEmitHelper::DefaultValue(FEmitterLocalContext& EmitterContext, const FEdGraphPinType& Type)
 {
-	return LiteralTerm(Type, FString(), nullptr);
+	return LiteralTerm(EmitterContext, Type, FString(), nullptr);
 }
 
 UFunction* FEmitHelper::GetOriginalFunction(UFunction* Function)
@@ -675,7 +682,7 @@ bool FEmitHelper::GenerateAssignmentCast(const FEdGraphPinType& LType, const FEd
 	return false;
 }
 
-void FEmitDefaultValueHelper::OuterGenerate(FDefaultValueHelperContext& Context
+void FEmitDefaultValueHelper::OuterGenerate(FEmitterLocalContext& Context
 	, const UProperty* Property
 	, const FString& OuterPath
 	, const uint8* DataContainer
@@ -736,7 +743,7 @@ FString FEmitDefaultValueHelper::GenerateGetDefaultValue(const UUserDefinedStruc
 	FStructOnScope StructData(Struct);
 	FStructureEditorUtils::Fill_MakeStructureDefaultValue(Struct, StructData.GetStructMemory());
 
-	FDefaultValueHelperContext Context;
+	FEmitterLocalContext Context;
 	Context.IncreaseIndent();
 	Context.IncreaseIndent();
 	for (auto Property : TFieldRange<const UProperty>(Struct))
@@ -748,9 +755,9 @@ FString FEmitDefaultValueHelper::GenerateGetDefaultValue(const UUserDefinedStruc
 	return Result;
 }
 
-void FEmitDefaultValueHelper::InnerGenerate(FDefaultValueHelperContext& Context, const UProperty* Property, const FString& PathToMember, const uint8* ValuePtr, const uint8* DefaultValuePtr, bool bWithoutFirstConstructionLine)
+void FEmitDefaultValueHelper::InnerGenerate(FEmitterLocalContext& Context, const UProperty* Property, const FString& PathToMember, const uint8* ValuePtr, const uint8* DefaultValuePtr, bool bWithoutFirstConstructionLine)
 {
-	auto OneLineConstruction = [](FDefaultValueHelperContext& LocalContext, const UProperty* LocalProperty, const uint8* LocalValuePtr, FString& OutSingleLine, bool bGenerateEmptyStructConstructor) -> bool
+	auto OneLineConstruction = [](FEmitterLocalContext& LocalContext, const UProperty* LocalProperty, const uint8* LocalValuePtr, FString& OutSingleLine, bool bGenerateEmptyStructConstructor) -> bool
 	{
 		bool bComplete = true;
 		FString ValueStr = HandleSpecialTypes(LocalContext, LocalProperty, LocalValuePtr);
@@ -763,7 +770,7 @@ void FEmitDefaultValueHelper::InnerGenerate(FDefaultValueHelperContext& Context,
 				check(StructProperty->Struct);
 				if (bGenerateEmptyStructConstructor)
 				{
-					ValueStr = FString::Printf(TEXT("F%s()"), *StructProperty->Struct->GetName()); //don;t override existing values
+					ValueStr = FString::Printf(TEXT("F%s{}"), *StructProperty->Struct->GetName()); //don;t override existing values
 				}
 				bComplete = false;
 			}
@@ -823,7 +830,7 @@ void FEmitDefaultValueHelper::InnerGenerate(FDefaultValueHelperContext& Context,
 	}
 }
 
-FString FEmitDefaultValueHelper::HandleSpecialTypes(FDefaultValueHelperContext& Context, const UProperty* Property, const uint8* ValuePtr)
+FString FEmitDefaultValueHelper::HandleSpecialTypes(FEmitterLocalContext& Context, const UProperty* Property, const uint8* ValuePtr)
 {
 	//TODO: Use Path maps for Objects
 	if (auto ObjectProperty = Cast<UObjectProperty>(Property))
@@ -831,17 +838,29 @@ FString FEmitDefaultValueHelper::HandleSpecialTypes(FDefaultValueHelperContext&
 		UObject* Object = ObjectProperty->GetPropertyValue(ValuePtr);
 		if (Object)
 		{
-			FString NativeName;
-			if (Context.FindLocalObject_InConstructor(Object, NativeName))
 			{
-				return NativeName;
+				const FString MappedObject = Context.FindGlobalObject(Object);
+				if (!MappedObject.IsEmpty())
+				{
+					return MappedObject;
+				}
+			}
+
+			{
+				FString LocalMappedObject;
+				if (Context.bInsideConstructor && Context.FindLocalObject_InConstructor(Object, LocalMappedObject))
+				{
+					return LocalMappedObject;
+				}
 			}
 
-			auto BPGC = Context.GetCurrentlyGeneratedClass();
-			auto CDO = BPGC ? BPGC->GetDefaultObject(false) : nullptr;
-			if (BPGC && Object && CDO && Object->IsIn(BPGC) && !Object->IsIn(CDO) && Context.bCreatingObjectsPerClass)
 			{
-				return HandleClassSubobject(Context, Object);
+				auto BPGC = Context.GetCurrentlyGeneratedClass();
+				auto CDO = BPGC ? BPGC->GetDefaultObject(false) : nullptr;
+				if (BPGC && Object && CDO && Object->IsIn(BPGC) && !Object->IsIn(CDO) && Context.bCreatingObjectsPerClass)
+				{
+					return HandleClassSubobject(Context, Object);
+				}
 			}
 
 			if (!Context.bCreatingObjectsPerClass && Property->HasAnyPropertyFlags(CPF_InstancedReference))
@@ -852,12 +871,10 @@ FString FEmitDefaultValueHelper::HandleSpecialTypes(FDefaultValueHelperContext&
 					return CreateAsInstancedSubobject;
 				}
 			}
-
-			const FString MappedObject = Context.FindGlobalObject(Object);
-			if (!MappedObject.IsEmpty())
-			{
-				return MappedObject;
-			}
+		}
+		else if (ObjectProperty->HasMetaData(FBlueprintMetadata::MD_LatentCallbackTarget))
+		{
+			return TEXT("this");
 		}
 	}
 
@@ -891,7 +908,7 @@ FString FEmitDefaultValueHelper::HandleSpecialTypes(FDefaultValueHelperContext&
 	return FString();
 }
 
-UActorComponent* FEmitDefaultValueHelper::HandleNonNativeComponent(FDefaultValueHelperContext& Context, UBlueprintGeneratedClass* BPGC, FName ObjectName, bool bNew, const FString& NativeName)
+UActorComponent* FEmitDefaultValueHelper::HandleNonNativeComponent(FEmitterLocalContext& Context, UBlueprintGeneratedClass* BPGC, FName ObjectName, bool bNew, const FString& NativeName)
 {
 	USCS_Node* Node = nullptr;
 	for (UBlueprintGeneratedClass* NodeOwner = BPGC; NodeOwner && !Node; NodeOwner = Cast<UBlueprintGeneratedClass>(NodeOwner->GetSuperClass()))
@@ -984,12 +1001,12 @@ FString FEmitDefaultValueHelper::GenerateConstructor(UClass* InBPGC)
 	auto BPGC = CastChecked<UBlueprintGeneratedClass>(InBPGC);
 	const FString CppClassName = FString(BPGC->GetPrefixCPP()) + BPGC->GetName();
 
-	FDefaultValueHelperContext Context;
+	FEmitterLocalContext Context;
 	Context.SetCurrentlyGeneratedClass(BPGC);
 	Context.AddLine(FString::Printf(TEXT("%s::%s(const FObjectInitializer& ObjectInitializer) : Super(ObjectInitializer)"), *CppClassName, *CppClassName));
 	Context.AddLine(TEXT("{"));
 	Context.IncreaseIndent();
-
+	Context.bInsideConstructor = true;
 	// When CDO is created create all subobjects owned by the class
 	TArray<UActorComponent*> ActorComponentTempatesOwnedByClass = BPGC->ComponentTemplates;
 	{
@@ -1104,11 +1121,11 @@ FString FEmitDefaultValueHelper::GenerateConstructor(UClass* InBPGC)
 	Context.AddLine(TEXT("}"));
 	Context.DecreaseIndent();
 	Context.AddLine(TEXT("}"));
-
+	Context.bInsideConstructor = false;
 	return Context.GetResult();
 }
 
-FString FEmitDefaultValueHelper::HandleClassSubobject(FDefaultValueHelperContext& Context, UObject* Object)
+FString FEmitDefaultValueHelper::HandleClassSubobject(FEmitterLocalContext& Context, UObject* Object)
 {
 	FString OuterStr;
 	if (!Context.FindLocalObject_InConstructor(Object->GetOuter(), OuterStr))
@@ -1140,7 +1157,7 @@ FString FEmitDefaultValueHelper::HandleClassSubobject(FDefaultValueHelperContext
 	return LocalNativeName;
 }
 
-FString FEmitDefaultValueHelper::HandleInstancedSubobject(FDefaultValueHelperContext& Context, UObject* Object)
+FString FEmitDefaultValueHelper::HandleInstancedSubobject(FEmitterLocalContext& Context, UObject* Object)
 {
 	check(Object);
 
diff --git a/Engine/Source/Developer/BlueprintCompilerCppBackend/Private/BlueprintCompilerCppBackendUtils.h b/Engine/Source/Developer/BlueprintCompilerCppBackend/Private/BlueprintCompilerCppBackendUtils.h
--- a/Engine/Source/Developer/BlueprintCompilerCppBackend/Private/BlueprintCompilerCppBackendUtils.h
+++ b/Engine/Source/Developer/BlueprintCompilerCppBackend/Private/BlueprintCompilerCppBackendUtils.h
@@ -7,16 +7,14 @@ class FBlueprintCompilerCppBackend;
 struct FSafeContextScopedEmmitter
 {
 private:
-	FString& Body;
+	FEmitterLocalContext& EmitterContext;
 	bool bSafeContextUsed;
-	FString CurrentIndent;
 public:
-	FString GetAdditionalIndent() const;
 	bool IsSafeContextUsed() const;
 
-	static FString ValidationChain(const FBPTerminal* Term, FBlueprintCompilerCppBackend& CppBackend);
+	static FString ValidationChain(FEmitterLocalContext& EmitterContext, const FBPTerminal* Term, FBlueprintCompilerCppBackend& CppBackend);
 
-	FSafeContextScopedEmmitter(FString& InBody, const FBPTerminal* Term, FBlueprintCompilerCppBackend& CppBackend, const TCHAR* InCurrentIndent);
+	FSafeContextScopedEmmitter(FEmitterLocalContext& InEmitterContext, const FBPTerminal* Term, FBlueprintCompilerCppBackend& CppBackend);
 	~FSafeContextScopedEmmitter();
 };
 
@@ -52,9 +50,9 @@ struct FEmitHelper
 
 	static FString GatherNativeHeadersToInclude(UField* SourceItem, const TArray<FString>& PersistentHeaders);
 
-	static FString LiteralTerm(const FEdGraphPinType& Type, const FString& CustomValue, UObject* LiteralObject);
+	static FString LiteralTerm(FEmitterLocalContext& EmitterContext, const FEdGraphPinType& Type, const FString& CustomValue, UObject* LiteralObject);
 
-	static FString DefaultValue(const FEdGraphPinType& Type);
+	static FString DefaultValue(FEmitterLocalContext& EmitterContext, const FEdGraphPinType& Type);
 
 	static UFunction* GetOriginalFunction(UFunction* Function);
 
@@ -65,7 +63,7 @@ struct FEmitHelper
 	static bool GenerateAssignmentCast(const FEdGraphPinType& LType, const FEdGraphPinType& RType, FString& OutCastBegin, FString& OutCastEnd);
 };
 
-struct FDefaultValueHelperContext
+struct FEmitterLocalContext
 {
 private:
 	FString Indent;
@@ -75,6 +73,7 @@ private:
 	TArray<FString> LowPriorityLines;
 	UClass* ActualClass;
 
+	//ConstructorOnly
 	TMap<UObject*, FString> NativeObjectNamesInConstructor;
 	TArray<UObject*> ObjectsCreatedPerClass;
 
@@ -82,16 +81,19 @@ private:
 public:
 
 	bool bCreatingObjectsPerClass;
+	bool bInsideConstructor;
 
-	FDefaultValueHelperContext()
+	FEmitterLocalContext()
 		: LocalNameIndexMax(0)
 		, ActualClass(nullptr)
 		, bCreatingObjectsPerClass(false)
+		, bInsideConstructor(false)
 	{}
 
 	// Functions to use in constructor only
 	bool FindLocalObject_InConstructor(UObject* Object, FString& OutNamePath) const
 	{
+		ensure(bInsideConstructor);
 		if (Object == ActualClass)
 		{
 			OutNamePath = TEXT("GetClass()");
@@ -109,11 +111,13 @@ public:
 
 	void AddObjectFromLocalProperty_InConstructor(UObject* Object, const FString& NativeName)
 	{
+		ensure(bInsideConstructor);
 		NativeObjectNamesInConstructor.Add(Object, NativeName);
 	}
 
 	FString AddNewObject_InConstructor(UObject* Object)
 	{
+		ensure(bInsideConstructor);
 		const FString UniqueName = GenerateUniqueLocalName();
 		NativeObjectNamesInConstructor.Add(Object, UniqueName);
 		if (bCreatingObjectsPerClass)
@@ -127,6 +131,7 @@ public:
 
 	const TArray<UObject*>& GetObjectsCreatedPerClass_InConstructor()
 	{
+		ensure(bInsideConstructor);
 		return ObjectsCreatedPerClass;
 	}
 
@@ -257,7 +262,6 @@ struct FEmitDefaultValueHelper
 	static FString GenerateGetDefaultValue(const UUserDefinedStruct* Struct);
 
 	static FString GenerateConstructor(UClass* BPGC);
-private:
 
 	enum class EPropertyAccessOperator
 	{
@@ -267,20 +271,21 @@ private:
 	};
 
 	// OuterPath ends context/outer name (or empty, if the scope is "this")
-	static void OuterGenerate(FDefaultValueHelperContext& Context, const UProperty* Property, const FString& OuterPath, const uint8* DataContainer, const uint8* OptionalDefaultDataContainer, EPropertyAccessOperator AccessOperator, bool bAllowProtected = false);
-	
+	static void OuterGenerate(FEmitterLocalContext& Context, const UProperty* Property, const FString& OuterPath, const uint8* DataContainer, const uint8* OptionalDefaultDataContainer, EPropertyAccessOperator AccessOperator, bool bAllowProtected = false);
+
+private:
 	// PathToMember ends with variable name
-	static void InnerGenerate(FDefaultValueHelperContext& Context, const UProperty* Property, const FString& PathToMember, const uint8* ValuePtr, const uint8* DefaultValuePtr, bool bWithoutFirstConstructionLine = false);
-	
+	static void InnerGenerate(FEmitterLocalContext& Context, const UProperty* Property, const FString& PathToMember, const uint8* ValuePtr, const uint8* DefaultValuePtr, bool bWithoutFirstConstructionLine = false);
+
 	// Returns native term, 
 	// returns empty string if cannot handle
-	static FString HandleSpecialTypes(FDefaultValueHelperContext& Context, const UProperty* Property, const uint8* ValuePtr);
+	static FString HandleSpecialTypes(FEmitterLocalContext& Context, const UProperty* Property, const uint8* ValuePtr);
 
-	static UActorComponent* HandleNonNativeComponent(FDefaultValueHelperContext& Context, UBlueprintGeneratedClass* BPGC, FName ObjectName, bool bNew, const FString& NativeName);
+	static UActorComponent* HandleNonNativeComponent(FEmitterLocalContext& Context, UBlueprintGeneratedClass* BPGC, FName ObjectName, bool bNew, const FString& NativeName);
 	
 	// Creates the subobject (of class) returns it's native local name, 
 	// returns empty string if cannot handle
-	static FString HandleClassSubobject(FDefaultValueHelperContext& Context, UObject* Object);
+	static FString HandleClassSubobject(FEmitterLocalContext& Context, UObject* Object);
 
-	static FString HandleInstancedSubobject(FDefaultValueHelperContext& Context, UObject* Object);
+	static FString HandleInstancedSubobject(FEmitterLocalContext& Context, UObject* Object);
 };
\ No newline at end of file
diff --git a/Engine/Source/Runtime/CoreUObject/Private/UObject/PropertyStruct.cpp b/Engine/Source/Runtime/CoreUObject/Private/UObject/PropertyStruct.cpp
--- a/Engine/Source/Runtime/CoreUObject/Private/UObject/PropertyStruct.cpp
+++ b/Engine/Source/Runtime/CoreUObject/Private/UObject/PropertyStruct.cpp
@@ -336,8 +336,15 @@ void UStructProperty::ExportTextItem( FString& ValueStr, const void* PropertyVal
 	UStructProperty_ExportTextItem(Struct, ValueStr, PropertyValue, DefaultValue, Parent, PortFlags, ExportRootScope);
 } 
 
-const TCHAR* UStructProperty::ImportText_Internal( const TCHAR* InBuffer, void* Data, int32 PortFlags, UObject* Parent, FOutputDevice* ErrorText ) const
+const TCHAR* UStructProperty::ImportText_Internal(const TCHAR* InBuffer, void* Data, int32 PortFlags, UObject* Parent, FOutputDevice* ErrorText) const
 {
+	return ImportText_Static(Struct, GetName(), InBuffer, Data, PortFlags, Parent, ErrorText);
+}
+
+const TCHAR* UStructProperty::ImportText_Static(UScriptStruct* InStruct, const FString& Name, const TCHAR* InBuffer, void* Data, int32 PortFlags, UObject* Parent, FOutputDevice* ErrorText)
+{
+	auto Struct = InStruct;
+
 	if (Struct->StructFlags & STRUCT_ImportTextItemNative)
 	{
 		UScriptStruct::ICppStructOps* CppStructOps = Struct->GetCppStructOps();
@@ -377,7 +384,7 @@ const TCHAR* UStructProperty::ImportText_Internal( const TCHAR* InBuffer, void*
 
 					if (*Buffer != TCHAR('\"'))
 					{
-						ErrorText->Logf(TEXT("%sImportText (%s): Bad quoted string at: %s"), ErrorCount++ > 0 ? LINE_TERMINATOR : TEXT(""), *GetName(), Buffer );
+						ErrorText->Logf(TEXT("%sImportText (%s): Bad quoted string at: %s"), ErrorCount++ > 0 ? LINE_TERMINATOR : TEXT(""), *Name, Buffer);
 						return NULL;
 					}
 				}
@@ -390,7 +397,7 @@ const TCHAR* UStructProperty::ImportText_Internal( const TCHAR* InBuffer, void*
 					SubCount--;
 					if( SubCount < 0 )
 					{
-						ErrorText->Logf(TEXT("%sImportText (%s): Too many closing parenthesis in: %s"), ErrorCount++ > 0 ? LINE_TERMINATOR : TEXT(""), *GetName(), InBuffer);
+						ErrorText->Logf(TEXT("%sImportText (%s): Too many closing parenthesis in: %s"), ErrorCount++ > 0 ? LINE_TERMINATOR : TEXT(""), *Name, InBuffer);
 						return NULL;
 					}
 				}
@@ -398,7 +405,7 @@ const TCHAR* UStructProperty::ImportText_Internal( const TCHAR* InBuffer, void*
 			}
 			if( SubCount > 0 )
 			{
-				ErrorText->Logf(TEXT("%sImportText(%s): Not enough closing parenthesis in: %s"), ErrorCount++ > 0 ? LINE_TERMINATOR : TEXT(""), *GetName(), InBuffer);
+				ErrorText->Logf(TEXT("%sImportText(%s): Not enough closing parenthesis in: %s"), ErrorCount++ > 0 ? LINE_TERMINATOR : TEXT(""), *Name, InBuffer);
 				return NULL;
 			}
 
@@ -410,7 +417,7 @@ const TCHAR* UStructProperty::ImportText_Internal( const TCHAR* InBuffer, void*
 			}
 			else if( *Buffer!=TCHAR(')') )
 			{
-				ErrorText->Logf(TEXT("%sImportText (%s): Missing closing parenthesis: %s"), ErrorCount++ > 0 ? LINE_TERMINATOR : TEXT(""), *GetName(), InBuffer);
+				ErrorText->Logf(TEXT("%sImportText (%s): Missing closing parenthesis: %s"), ErrorCount++ > 0 ? LINE_TERMINATOR : TEXT(""), *Name, InBuffer);
 				return NULL;
 			}
 
@@ -422,7 +429,7 @@ const TCHAR* UStructProperty::ImportText_Internal( const TCHAR* InBuffer, void*
 	}
 	else
 	{
-		ErrorText->Logf(TEXT("%sImportText (%s): Missing opening parenthesis: %s"), ErrorCount++ > 0 ? LINE_TERMINATOR : TEXT(""), *GetName(), InBuffer);
+		ErrorText->Logf(TEXT("%sImportText (%s): Missing opening parenthesis: %s"), ErrorCount++ > 0 ? LINE_TERMINATOR : TEXT(""), *Name, InBuffer);
 		return NULL;
 	}
 	return Buffer;
diff --git a/Engine/Source/Runtime/CoreUObject/Public/UObject/UnrealType.h b/Engine/Source/Runtime/CoreUObject/Public/UObject/UnrealType.h
--- a/Engine/Source/Runtime/CoreUObject/Public/UObject/UnrealType.h
+++ b/Engine/Source/Runtime/CoreUObject/Public/UObject/UnrealType.h
@@ -3203,6 +3203,8 @@ public:
 	virtual bool SameType(const UProperty* Other) const override;
 	// End of UProperty interface
 
+	static const TCHAR* ImportText_Static(UScriptStruct* InStruct, const FString& InName, const TCHAR* Buffer, void* Data, int32 PortFlags, UObject* OwnerObject, FOutputDevice* ErrorText);
+
 	bool UseNativeSerialization() const;
 	bool UseBinarySerialization(const FArchive& Ar) const;
 	bool UseBinaryOrNativeSerialization(const FArchive& Ar) const;
