UEBP-40 C++ code generation (WIP):
- code was split into many files and refactored
- return and infinite loop fix
- generated code uses "GloballyMappedObjects" instead "LoadObject"
- various improvements

diff --git a/Engine/Source/Developer/BlueprintCompilerCppBackend/Private/BlueprintCompilerCppBackend.cpp b/Engine/Source/Developer/BlueprintCompilerCppBackend/Private/BlueprintCompilerCppBackend.cpp
--- a/Engine/Source/Developer/BlueprintCompilerCppBackend/Private/BlueprintCompilerCppBackend.cpp
+++ b/Engine/Source/Developer/BlueprintCompilerCppBackend/Private/BlueprintCompilerCppBackend.cpp
@@ -5,223 +5,71 @@
 #include "BlueprintCompilerCppBackendUtils.h"
 #include "EdGraphSchema_K2.h"
 
-FString FBlueprintCompilerCppBackend::TermToText(FEmitterLocalContext& EmitterContext, const FBPTerminal* Term, bool bUseSafeContext)
-{
-	const FString PSC_Self(TEXT("self"));
-	if (Term->bIsLiteral)
-	{
-		return FEmitHelper::LiteralTerm(EmitterContext, Term->Type, Term->Name, Term->ObjectLiteral);
-	}
-	else if (Term->InlineGeneratedParameter)
-	{
-		if (KCST_SwitchValue == Term->InlineGeneratedParameter->Type)
-		{
-			return EmitSwitchValueStatmentInner(EmitterContext, *Term->InlineGeneratedParameter);
-		}
-		else if (KCST_CallFunction == Term->InlineGeneratedParameter->Type)
-		{
-			return EmitCallStatmentInner(EmitterContext, *Term->InlineGeneratedParameter, true);
-		}
-		else
-		{
-			ensureMsgf(false, TEXT("KCST %d is not accepted as inline statement."), Term->InlineGeneratedParameter->Type);
-			return FString();
-		}
-	}
-	else
-	{
-		FString ResultPath(TEXT(""));
-		if ((Term->Context != NULL) && (Term->Context->Name != PSC_Self))
-		{
-			ResultPath = TermToText(EmitterContext, Term->Context, false);
-
-			if (Term->Context->IsStructContextType())
-			{
-				ResultPath += TEXT(".");
-			}
-			else
-			{
-				ResultPath += TEXT("->");
-			}
-		}
-
-		FString Conditions;
-		if (bUseSafeContext)
-		{
-			Conditions = FSafeContextScopedEmmitter::ValidationChain(EmitterContext, Term->Context, *this);
-		}
-
-		ResultPath += Term->AssociatedVarProperty ? Term->AssociatedVarProperty->GetNameCPP() : Term->Name;
-		return Conditions.IsEmpty()
-			? ResultPath
-			: FString::Printf(TEXT("((%s) ? (%s) : (%s))"), *Conditions, *ResultPath, *FEmitHelper::DefaultValue(EmitterContext, Term->Type));
-	}
-}
-
-FString FBlueprintCompilerCppBackend::LatentFunctionInfoTermToText(FEmitterLocalContext& EmitterContext, FBPTerminal* Term, FBlueprintCompiledStatement* TargetLabel)
+// Generates single "if" scope. Its condition checks context of given term.
+struct FSafeContextScopedEmmitter
 {
-	auto LatentInfoStruct = FLatentActionInfo::StaticStruct();
-
-	// Find the term name we need to fixup
-	FString FixupTermName;
-	for (UProperty* Prop = LatentInfoStruct->PropertyLink; Prop; Prop = Prop->PropertyLinkNext)
-	{
-		static const FName NeedsLatentFixup(TEXT("NeedsLatentFixup"));
-		if (Prop->GetBoolMetaData(NeedsLatentFixup))
-		{
-			FixupTermName = Prop->GetName();
-			break;
-		}
-	}
-
-	check(!FixupTermName.IsEmpty());
-
-	FString StructValues = Term->Name;
-
-	// Index 0 is always the ubergraph
-	const int32 TargetStateIndex = StateMapPerFunction[0].StatementToStateIndex(TargetLabel);
-	const int32 LinkageTermStartIdx = StructValues.Find(FixupTermName);
-	check(LinkageTermStartIdx != INDEX_NONE);
-	StructValues = StructValues.Replace(TEXT("-1"), *FString::FromInt(TargetStateIndex));
+private:
+	FEmitterLocalContext& EmitterContext;
+	bool bSafeContextUsed;
+public:
 
-	return FEmitHelper::LiteralTerm(EmitterContext, Term->Type, StructValues, nullptr);
-}
-
-void FBlueprintCompilerCppBackend::EmitStructProperties(FStringOutputDevice& Target, UStruct* SourceClass)
-{
-	// Emit class variables
-	for (TFieldIterator<UProperty> It(SourceClass, EFieldIteratorFlags::ExcludeSuper); It; ++It)
+	bool IsSafeContextUsed() const
 	{
-		UProperty* Property = *It;
-		check(Property);
-
-		Emit(Header, TEXT("\n\tUPROPERTY("));
-		{
-			TArray<FString> Tags = FEmitHelper::ProperyFlagsToTags(Property->PropertyFlags);
-			Tags.Emplace(FEmitHelper::HandleRepNotifyFunc(Property));
-			Tags.Emplace(FEmitHelper::HandleMetaData(Property, false));
-			Tags.Remove(FString());
-
-			FString AllTags;
-			FEmitHelper::ArrayToString(Tags, AllTags, TEXT(", "));
-			Emit(Header, *AllTags);
-		}
-		Emit(Header, TEXT(")\n"));
-		Emit(Header, TEXT("\t"));
-		Property->ExportCppDeclaration(Target, EExportedDeclaration::Member, NULL, EPropertyExportCPPFlags::CPPF_CustomTypeName | EPropertyExportCPPFlags::CPPF_BlueprintCppBackend);
-		Emit(Header, TEXT(";\n"));
+		return bSafeContextUsed;
 	}
-}
 
-void FBlueprintCompilerCppBackend::GenerateCodeFromClass(UClass* SourceClass, TIndirectArray<FKismetFunctionContext>& Functions, bool bGenerateStubsOnly)
-{
-	auto CleanCppClassName = SourceClass->GetName();
-	CppClassName = FString(SourceClass->GetPrefixCPP()) + CleanCppClassName;
-	
-	EmitFileBeginning(CleanCppClassName, SourceClass);
-
-	// MC DELEGATE DECLARATION
+	static FString ValidationChain(FEmitterLocalContext& EmitterContext, const FBPTerminal* Term, FBlueprintCompilerCppBackend& CppBackend)
 	{
-		auto DelegateDeclarations = FEmitHelper::EmitMulticastDelegateDeclarations(SourceClass);
-		FString AllDeclarations;
-		FEmitHelper::ArrayToString(DelegateDeclarations, AllDeclarations, TEXT(";\n"));
-		if (DelegateDeclarations.Num())
+		TArray<FString> SafetyConditions;
+		for (; Term; Term = Term->Context)
 		{
-			Emit(Header, *AllDeclarations);
-			Emit(Header, TEXT(";\n"));
-		}
-	}
-
-	// GATHER ALL SC DELEGATES
-	{
-		TArray<UDelegateProperty*> Delegates;
-		for (TFieldIterator<UDelegateProperty> It(SourceClass, EFieldIteratorFlags::ExcludeSuper); It; ++It)
-		{
-			Delegates.Add(*It);
-		}
-
-		for (auto& FuncContext : Functions)
-		{
-			for (TFieldIterator<UDelegateProperty> It(FuncContext.Function, EFieldIteratorFlags::ExcludeSuper); It; ++It)
+			if (!Term->IsStructContextType() && (Term->Type.PinSubCategory != TEXT("self")))
 			{
-				Delegates.Add(*It);
+				ensure(!Term->bIsLiteral);
+				SafetyConditions.Add(CppBackend.TermToText(EmitterContext, Term, false));
 			}
 		}
 
-		auto DelegateDeclarations = FEmitHelper::EmitSinglecastDelegateDeclarations(Delegates);
-		FString AllDeclarations;
-		FEmitHelper::ArrayToString(DelegateDeclarations, AllDeclarations, TEXT(";\n"));
-		if (DelegateDeclarations.Num())
-		{
-			Emit(Header, *AllDeclarations);
-			Emit(Header, TEXT(";\n"));
-		}
-	}
-
-	// Class declaration
-	const bool bIsInterface = SourceClass->IsChildOf<UInterface>();
-	if (bIsInterface)
-	{
-		Emit(Header, *FString::Printf(TEXT("UINTERFACE(Blueprintable)\nclass U%s : public UInterface\n{\n\tGENERATED_BODY()\n};\n"), *CleanCppClassName));
-		Emit(Header, *FString::Printf(TEXT("\nclass I%s"), *CleanCppClassName));
-	}
-	else
-	{
-		Emit(Header, TEXT("UCLASS("));
-		if (!SourceClass->IsChildOf<UBlueprintFunctionLibrary>())
-		{
-			Emit(Header, TEXT("Blueprintable, BlueprintType"));
-		}
-		Emit(Header, TEXT(")\n"));
-
-		UClass* SuperClass = SourceClass->GetSuperClass();
-		Emit(Header, *FString::Printf(TEXT("class %s : public %s%s"), *CppClassName, SuperClass->GetPrefixCPP(), *SuperClass->GetName()));
-
-		for (auto& ImplementedInterface : SourceClass->Interfaces)
+		FString Result;
+		for (int32 Iter = SafetyConditions.Num() - 1; Iter >= 0; --Iter)
 		{
-			if (ImplementedInterface.Class)
+			Result += FString(TEXT("IsValid("));
+			Result += SafetyConditions[Iter];
+			Result += FString(TEXT(")"));
+			if (Iter)
 			{
-				Emit(Header, *FString::Printf(TEXT(", public I%s"), *ImplementedInterface.Class->GetName()));
+				Result += FString(TEXT(" && "));
 			}
 		}
-	}
-
-	// Begin scope
-	Emit(Header,TEXT("\n{\npublic:\n\tGENERATED_BODY()\n"));
-
-	EmitStructProperties(Header, SourceClass);
 
-	// Create the state map
-	for (int32 i = 0; i < Functions.Num(); ++i)
-	{
-		StateMapPerFunction.Add(FFunctionLabelInfo());
-		FunctionIndexMap.Add(&Functions[i], i);
+		return Result;
 	}
 
-	// Emit function declarations and definitions (writes to header and body simultaneously)
-	if (Functions.Num() > 0)
+	FSafeContextScopedEmmitter(FEmitterLocalContext& InEmitterContext, const FBPTerminal* Term, FBlueprintCompilerCppBackend& CppBackend)
+		: EmitterContext(InEmitterContext)
+		, bSafeContextUsed(false)
 	{
-		Emit(Header, TEXT("\n"));
-	}
+		const FString Conditions = ValidationChain(EmitterContext, Term, CppBackend);
 
-	if (!bIsInterface)
-	{
-		Emit(Header, *FString::Printf(TEXT("\t%s(const FObjectInitializer& ObjectInitializer = FObjectInitializer::Get());\n\n"), *CppClassName));
-		Emit(Body, *FEmitDefaultValueHelper::GenerateConstructor(SourceClass));
+		if (!Conditions.IsEmpty())
+		{
+			bSafeContextUsed = true;
+			EmitterContext.AddLine(FString::Printf(TEXT("if(%s)"), *Conditions));
+			EmitterContext.AddLine(TEXT("{"));
+			EmitterContext.IncreaseIndent();
+		}
 	}
 
-	for (int32 i = 0; i < Functions.Num(); ++i)
+	~FSafeContextScopedEmmitter()
 	{
-		if (Functions[i].IsValid())
+		if (bSafeContextUsed)
 		{
-			ConstructFunction(Functions[i], bGenerateStubsOnly);
+			EmitterContext.DecreaseIndent();
+			EmitterContext.AddLine(TEXT("}"));
 		}
 	}
 
-	Emit(Header, TEXT("};\n\n"));
-
-	Emit(Body, *FEmitHelper::EmitLifetimeReplicatedPropsImpl(SourceClass, CppClassName, TEXT("")));
-}
+};
 
 void FBlueprintCompilerCppBackend::EmitCallDelegateStatment(FEmitterLocalContext& EmitterContext, FKismetFunctionContext& FunctionContext, FBlueprintCompiledStatement& Statement)
 {
@@ -230,145 +78,6 @@ void FBlueprintCompilerCppBackend::EmitCallDelegateStatment(FEmitterLocalContext
 	EmitterContext.AddLine(FString::Printf(TEXT("%s.Broadcast(%s);"), *TermToText(EmitterContext, Statement.FunctionContext, false), *EmitMethodInputParameterList(EmitterContext, Statement)));
 }
 
-FString FBlueprintCompilerCppBackend::EmitMethodInputParameterList(FEmitterLocalContext& EmitterContext, FBlueprintCompiledStatement& Statement)
-{
-	FString Result;
-	int32 NumParams = 0;
-	for (TFieldIterator<UProperty> PropIt(Statement.FunctionToCall); PropIt && (PropIt->PropertyFlags & CPF_Parm); ++PropIt)
-	{
-		UProperty* FuncParamProperty = *PropIt;
-
-		if (!FuncParamProperty->HasAnyPropertyFlags(CPF_ReturnParm))
-		{
-			if (NumParams > 0)
-			{
-				Result += TEXT(", ");
-			}
-
-			FString VarName;
-
-			FBPTerminal* Term = Statement.RHS[NumParams];
-			ensure(Term != NULL);
-
-			if ((Statement.TargetLabel != NULL) && (Statement.UbergraphCallIndex == NumParams))
-			{
-				// The target label will only ever be set on a call function when calling into the Ubergraph or
-				// on a latent function that will later call into the ubergraph, either of which requires a patchup
-				UStructProperty* StructProp = Cast<UStructProperty>(FuncParamProperty);
-				if (StructProp && (StructProp->Struct == FLatentActionInfo::StaticStruct()))
-				{
-					// Latent function info case
-					VarName = LatentFunctionInfoTermToText(EmitterContext, Term, Statement.TargetLabel);
-				}
-				else
-				{
-					// Ubergraph entry point case
-					VarName = FString::FromInt(StateMapPerFunction[0].StatementToStateIndex(Statement.TargetLabel));
-				}
-			}
-			else
-			{
-				// Emit a normal parameter term
-				VarName = TermToText(EmitterContext, Term);
-			}
-
-			if (FuncParamProperty->HasAnyPropertyFlags(CPF_OutParm))
-			{
-				Result += TEXT("/*out*/ ");
-			}
-			Result += *VarName;
-
-			NumParams++;
-		}
-	}
-
-	return Result;
-}
-
-FString FBlueprintCompilerCppBackend::EmitCallStatmentInner(FEmitterLocalContext& EmitterContext, FBlueprintCompiledStatement& Statement, bool bInline)
-{
-	const bool bCallOnDifferentObject = Statement.FunctionContext && (Statement.FunctionContext->Name != TEXT("self"));
-	const bool bStaticCall = Statement.FunctionToCall->HasAnyFunctionFlags(FUNC_Static);
-	const bool bUseSafeContext = bCallOnDifferentObject && !bStaticCall;
-	const bool bInterfaceCall = bCallOnDifferentObject && Statement.FunctionContext && (UEdGraphSchema_K2::PC_Interface == Statement.FunctionContext->Type.PinCategory);
-
-	FString Result;
-	bool bCloseCast = false;
-	if (!bInline)
-	{
-		// Handle the return value of the function being called
-		UProperty* FuncToCallReturnProperty = Statement.FunctionToCall->GetReturnProperty();
-		if (FuncToCallReturnProperty && ensure(Statement.LHS))
-		{
-			FString BeginCast;
-			if (auto ObjectProperty = Cast<UObjectProperty>(FuncToCallReturnProperty))
-			{
-				UClass* LClass = Statement.LHS ? Cast<UClass>(Statement.LHS->Type.PinSubCategoryObject.Get()) : nullptr;
-				if (LClass && LClass->IsChildOf(ObjectProperty->PropertyClass) && !ObjectProperty->PropertyClass->IsChildOf(LClass))
-				{
-					BeginCast = FString::Printf(TEXT("CastChecked<%s%s>("), LClass->GetPrefixCPP(), *LClass->GetName());
-					bCloseCast = true;
-				}
-			}
-
-			Result += FString::Printf(TEXT("%s = %s"), *TermToText(EmitterContext, Statement.LHS), *BeginCast);
-				
-		}
-	}
-
-	// Emit object to call the method on
-	if (bInterfaceCall)
-	{
-		auto ContextInterfaceClass = CastChecked<UClass>(Statement.FunctionContext->Type.PinSubCategoryObject.Get());
-		ensure(ContextInterfaceClass->IsChildOf<UInterface>());
-		Result += FString::Printf(TEXT("I%s::Execute_%s(%s.GetObject(), ")
-			, *ContextInterfaceClass->GetName()
-			, *Statement.FunctionToCall->GetName()
-			, *TermToText(EmitterContext, Statement.FunctionContext, false));
-	}
-	else
-	{
-		if (bStaticCall)
-		{
-			const bool bIsCustomThunk = Statement.FunctionToCall->HasMetaData(TEXT("CustomStructureParam")) || Statement.FunctionToCall->HasMetaData(TEXT("ArrayParm"));
-			auto OwnerClass = Statement.FunctionToCall->GetOuterUClass();
-			Result += bIsCustomThunk ? TEXT("FCustomThunkTemplates::") : FString::Printf(TEXT("%s%s::"), OwnerClass->GetPrefixCPP(), *OwnerClass->GetName());
-		}
-		else if (bCallOnDifferentObject) //@TODO: Badness, could be a self reference wired to another instance!
-		{
-			Result += FString::Printf(TEXT("%s->"), *TermToText(EmitterContext, Statement.FunctionContext, false));
-		}
-
-		if (Statement.bIsParentContext)
-		{
-			Result += TEXT("Super::");
-		}
-		Result += Statement.FunctionToCall->GetName();
-		if (Statement.bIsParentContext && FEmitHelper::ShoulsHandleAsNativeEvent(Statement.FunctionToCall))
-		{
-			ensure(!bCallOnDifferentObject);
-			Result += TEXT("_Implementation");
-		}
-
-		// Emit method parameter list
-		Result += TEXT("(");
-	}
-	Result += EmitMethodInputParameterList(EmitterContext, Statement);
-	Result += TEXT(")");
-
-	if (bCloseCast)
-	{
-		Result += TEXT(", ECastCheckedType::NullAllowed)");
-	}
-
-	if (!bInline)
-	{
-		Result += TEXT(";");
-	}
-
-	return Result;
-}
-
 void FBlueprintCompilerCppBackend::EmitCallStatment(FEmitterLocalContext& EmitterContext, FKismetFunctionContext& FunctionContext, FBlueprintCompiledStatement& Statement)
 {
 	const bool bCallOnDifferentObject = Statement.FunctionContext && (Statement.FunctionContext->Name != TEXT("self"));
@@ -381,46 +90,6 @@ void FBlueprintCompilerCppBackend::EmitCallStatment(FEmitterLocalContext& Emitte
 	}
 }
 
-FString FBlueprintCompilerCppBackend::EmitSwitchValueStatmentInner(FEmitterLocalContext& EmitterContext, FBlueprintCompiledStatement& Statement)
-{
-	check(Statement.RHS.Num() >= 2);
-	const int32 TermsBeforeCases = 1;
-	const int32 TermsPerCase = 2;
-	const int32 NumCases = ((Statement.RHS.Num() - 2) / TermsPerCase);
-	auto IndexTerm = Statement.RHS[0];
-	auto DefaultValueTerm = Statement.RHS.Last();
-
-	FString Result = FString::Printf(TEXT("TSwitchValue(%s, %s, %d")
-		, *TermToText(EmitterContext, IndexTerm) //index
-		, *TermToText(EmitterContext, DefaultValueTerm) // default
-		, NumCases);
-	
-	const uint32 CppTemplateTypeFlags = EPropertyExportCPPFlags::CPPF_CustomTypeName 
-		| EPropertyExportCPPFlags::CPPF_NoConst | EPropertyExportCPPFlags::CPPF_NoRef 
-		| EPropertyExportCPPFlags::CPPF_BlueprintCppBackend;
-
-	FStringOutputDevice IndexDeclaration;
-	check(IndexTerm && IndexTerm->AssociatedVarProperty);
-	IndexTerm->AssociatedVarProperty->ExportCppDeclaration(IndexDeclaration, EExportedDeclaration::Parameter, NULL, CppTemplateTypeFlags, true);
-
-	check(DefaultValueTerm && DefaultValueTerm->AssociatedVarProperty);
-	FStringOutputDevice ValueDeclaration;
-	DefaultValueTerm->AssociatedVarProperty->ExportCppDeclaration(ValueDeclaration, EExportedDeclaration::Parameter, NULL, CppTemplateTypeFlags, true);
-
-	for (int32 TermIndex = TermsBeforeCases; TermIndex < (NumCases * TermsPerCase); TermIndex += TermsPerCase)
-	{
-		Result += FString::Printf(TEXT(", TSwitchPair<%s, %s>(%s, %s)")
-			, *IndexDeclaration
-			, *ValueDeclaration
-			, *TermToText(EmitterContext, Statement.RHS[TermIndex])
-			, *TermToText(EmitterContext, Statement.RHS[TermIndex + 1]));
-	}
-
-	Result += TEXT(")");
-
-	return Result;
-}
-
 void FBlueprintCompilerCppBackend::EmitAssignmentStatment(FEmitterLocalContext& EmitterContext, FKismetFunctionContext& FunctionContext, FBlueprintCompiledStatement& Statement)
 {
 	check(Statement.LHS && Statement.RHS[0]);
@@ -619,521 +288,432 @@ void FBlueprintCompilerCppBackend::EmitPushStateStatement(FEmitterLocalContext&
 	EmitterContext.AddLine(FString::Printf(TEXT("StateStack.Push(%d);"), StatementToStateIndex(FunctionContext, Statement.TargetLabel)));
 }
 
-void FBlueprintCompilerCppBackend::EmitEndOfThreadStatement(FEmitterLocalContext& EmitterContext, FKismetFunctionContext& FunctionContext, const FString& ReturnValueString)
+void FBlueprintCompilerCppBackend::EmitEndOfThreadStatement(FEmitterLocalContext& EmitterContext, FKismetFunctionContext& FunctionContext)
 {
 	ensure(FunctionContext.bUseFlowStack);
 	EmitterContext.AddLine(TEXT("CurrentState = (StateStack.Num() > 0) ? StateStack.Pop(/*bAllowShrinking=*/ false) : -1;"));
 	EmitterContext.AddLine(TEXT("break;"));
 }
 
-void FBlueprintCompilerCppBackend::EmitReturnStatement(FKismetFunctionContext& FunctionContext, const FString& ReturnValueString)
+FString FBlueprintCompilerCppBackend::EmitSwitchValueStatmentInner(FEmitterLocalContext& EmitterContext, FBlueprintCompiledStatement& Statement)
 {
-	Emit(Body, *FString::Printf(TEXT("\treturn%s;\n"), *ReturnValueString));
-}
+	check(Statement.RHS.Num() >= 2);
+	const int32 TermsBeforeCases = 1;
+	const int32 TermsPerCase = 2;
+	const int32 NumCases = ((Statement.RHS.Num() - 2) / TermsPerCase);
+	auto IndexTerm = Statement.RHS[0];
+	auto DefaultValueTerm = Statement.RHS.Last();
 
-void FBlueprintCompilerCppBackend::DeclareLocalVariables(FKismetFunctionContext& FunctionContext, TArray<UProperty*>& LocalVariables)
-{
-	for (int32 i = 0; i < LocalVariables.Num(); ++i)
-	{
-		UProperty* LocalVariable = LocalVariables[i];
+	FString Result = FString::Printf(TEXT("TSwitchValue(%s, %s, %d")
+		, *TermToText(EmitterContext, IndexTerm) //index
+		, *TermToText(EmitterContext, DefaultValueTerm) // default
+		, NumCases);
 
-		Emit(Body, TEXT("\t"));
-		LocalVariable->ExportCppDeclaration(Body, EExportedDeclaration::Local, NULL, EPropertyExportCPPFlags::CPPF_CustomTypeName | EPropertyExportCPPFlags::CPPF_BlueprintCppBackend);
-		Emit(Body, TEXT("{};\n"));
-	}
+	const uint32 CppTemplateTypeFlags = EPropertyExportCPPFlags::CPPF_CustomTypeName
+		| EPropertyExportCPPFlags::CPPF_NoConst | EPropertyExportCPPFlags::CPPF_NoRef
+		| EPropertyExportCPPFlags::CPPF_BlueprintCppBackend;
 
-	if (LocalVariables.Num() > 0)
-	{
-		Emit(Body, TEXT("\n"));
-	}
-}
+	FStringOutputDevice IndexDeclaration;
+	check(IndexTerm && IndexTerm->AssociatedVarProperty);
+	IndexTerm->AssociatedVarProperty->ExportCppDeclaration(IndexDeclaration, EExportedDeclaration::Parameter, NULL, CppTemplateTypeFlags, true);
 
-void FBlueprintCompilerCppBackend::DeclareStateSwitch(FKismetFunctionContext& FunctionContext)
-{
-	if (FunctionContext.bUseFlowStack)
-	{
-		Emit(Body, TEXT("\tTArray< int32, TInlineAllocator<8> > StateStack;\n"));
-	}
-	Emit(Body, TEXT("\tint32 CurrentState = 0;\n"));
-	Emit(Body, TEXT("\tdo\n"));
-	Emit(Body, TEXT("\t{\n"));
-	Emit(Body, TEXT("\t\tswitch( CurrentState )\n"));
-	Emit(Body, TEXT("\t\t{\n"));
-	Emit(Body, TEXT("\t\tcase 0:\n"));
-	Emit(Body, TEXT("\t\t\t{\n"));
-}
+	check(DefaultValueTerm && DefaultValueTerm->AssociatedVarProperty);
+	FStringOutputDevice ValueDeclaration;
+	DefaultValueTerm->AssociatedVarProperty->ExportCppDeclaration(ValueDeclaration, EExportedDeclaration::Parameter, NULL, CppTemplateTypeFlags, true);
 
-void FBlueprintCompilerCppBackend::CloseStateSwitch(FKismetFunctionContext& FunctionContext)
-{
-	// Default error-catching case 
-	Emit(Body, TEXT("\t\t\t}\n"));
-	Emit(Body, TEXT("\t\tdefault:\n"));
-	if (FunctionContext.bUseFlowStack)
+	for (int32 TermIndex = TermsBeforeCases; TermIndex < (NumCases * TermsPerCase); TermIndex += TermsPerCase)
 	{
-		Emit(Body, TEXT("\t\t\tcheck(false); // Invalid state\n"));
+		Result += FString::Printf(TEXT(", TSwitchPair<%s, %s>(%s, %s)")
+			, *IndexDeclaration
+			, *ValueDeclaration
+			, *TermToText(EmitterContext, Statement.RHS[TermIndex])
+			, *TermToText(EmitterContext, Statement.RHS[TermIndex + 1]));
 	}
-	Emit(Body, TEXT("\t\t\tbreak;\n"));
 
-	// Close the switch block and do-while loop
-	Emit(Body, TEXT("\t\t}\n"));
-	Emit(Body, TEXT("\t} while( CurrentState != -1 );\n"));
+	Result += TEXT(")");
+
+	return Result;
 }
 
-void FBlueprintCompilerCppBackend::ConstructFunction(FKismetFunctionContext& FunctionContext, bool bGenerateStubOnly)
+FString FBlueprintCompilerCppBackend::EmitMethodInputParameterList(FEmitterLocalContext& EmitterContext, FBlueprintCompiledStatement& Statement)
 {
-	if (FunctionContext.IsDelegateSignature())
+	FString Result;
+	int32 NumParams = 0;
+	for (TFieldIterator<UProperty> PropIt(Statement.FunctionToCall); PropIt && (PropIt->PropertyFlags & CPF_Parm); ++PropIt)
 	{
-		return;
-	}
-
-	UFunction* Function = FunctionContext.Function;
+		UProperty* FuncParamProperty = *PropIt;
 
-	UProperty* ReturnValue = NULL;
-	TArray<UProperty*> LocalVariables;
+		if (!FuncParamProperty->HasAnyPropertyFlags(CPF_ReturnParm))
+		{
+			if (NumParams > 0)
+			{
+				Result += TEXT(", ");
+			}
 
-	{
-		FString FunctionName;
-		Function->GetName(FunctionName);
+			FString VarName;
 
-		TArray<UProperty*> ArgumentList;
+			FBPTerminal* Term = Statement.RHS[NumParams];
+			ensure(Term != NULL);
 
-		// Split the function property list into arguments, a return value (if any), and local variable declarations
-		for (TFieldIterator<UProperty> It(Function); It; ++It)
-		{
-			UProperty* Property = *It;
-			if (Property->HasAnyPropertyFlags(CPF_Parm))
+			if ((Statement.TargetLabel != NULL) && (Statement.UbergraphCallIndex == NumParams))
 			{
-				if (Property->HasAnyPropertyFlags(CPF_ReturnParm))
+				// The target label will only ever be set on a call function when calling into the Ubergraph or
+				// on a latent function that will later call into the ubergraph, either of which requires a patchup
+				UStructProperty* StructProp = Cast<UStructProperty>(FuncParamProperty);
+				if (StructProp && (StructProp->Struct == FLatentActionInfo::StaticStruct()))
 				{
-					if (ReturnValue == NULL)
-					{
-						ReturnValue = Property;
-						LocalVariables.Add(Property);
-					}
-					else
-					{
-						UE_LOG(LogK2Compiler, Error, TEXT("Function %s from graph @@ has more than one return value (named %s and %s)"),
-							*FunctionName, *GetPathNameSafe(FunctionContext.SourceGraph), *ReturnValue->GetName(), *Property->GetName());
-					}
+					// Latent function info case
+					VarName = LatentFunctionInfoTermToText(EmitterContext, Term, Statement.TargetLabel);
 				}
 				else
 				{
-					ArgumentList.Add(Property);
+					// Ubergraph entry point case
+					VarName = FString::FromInt(StateMapPerFunction[0].StatementToStateIndex(Statement.TargetLabel));
 				}
 			}
 			else
 			{
-				LocalVariables.Add(Property);
-			}
-		}
-
-		bool bAddCppFromBpEventMD = false;
-		bool bGenerateAsNativeEventImplementation = false;
-		// Get original function declaration
-		if (FEmitHelper::ShoulsHandleAsNativeEvent(Function)) // BlueprintNativeEvent
-		{
-			bGenerateAsNativeEventImplementation = true;
-			FunctionName += TEXT("_Implementation");
-		}
-		else if (FEmitHelper::ShoulsHandleAsImplementableEvent(Function)) // BlueprintImplementableEvent
-		{
-			bAddCppFromBpEventMD = true;
-		}
-
-		// Emit the declaration
-		const FString ReturnType = ReturnValue ? ReturnValue->GetCPPType(NULL, EPropertyExportCPPFlags::CPPF_CustomTypeName) : TEXT("void");
-
-		//@TODO: Make the header+body export more uniform
-		{
-			const FString Start = FString::Printf(TEXT("%s %s%s%s("), *ReturnType, TEXT("%s"), TEXT("%s"), *FunctionName);
-
-			TArray<FString> AdditionalMetaData;
-			if (bAddCppFromBpEventMD)
-			{
-				AdditionalMetaData.Emplace(TEXT("CppFromBpEvent"));
+				// Emit a normal parameter term
+				VarName = TermToText(EmitterContext, Term);
 			}
 
-			if (!bGenerateAsNativeEventImplementation)
+			if (FuncParamProperty->HasAnyPropertyFlags(CPF_OutParm))
 			{
-				Emit(Header, *FString::Printf(TEXT("\t%s\n"), *FEmitHelper::EmitUFuntion(Function, &AdditionalMetaData)));
+				Result += TEXT("/*out*/ ");
 			}
-			Emit(Header, TEXT("\t"));
+			Result += *VarName;
 
-			if (Function->HasAllFunctionFlags(FUNC_Static))
-			{
-				Emit(Header, TEXT("static "));
-			}
+			NumParams++;
+		}
+	}
 
-			if (bGenerateAsNativeEventImplementation)
-			{
-				Emit(Header, TEXT("virtual "));
-			}
+	return Result;
+}
 
-			Emit(Header, *FString::Printf(*Start, TEXT(""), TEXT("")));
-			Emit(Body, *FString::Printf(*Start, *CppClassName, TEXT("::")));
+FString FBlueprintCompilerCppBackend::EmitCallStatmentInner(FEmitterLocalContext& EmitterContext, FBlueprintCompiledStatement& Statement, bool bInline)
+{
+	const bool bCallOnDifferentObject = Statement.FunctionContext && (Statement.FunctionContext->Name != TEXT("self"));
+	const bool bStaticCall = Statement.FunctionToCall->HasAnyFunctionFlags(FUNC_Static);
+	const bool bUseSafeContext = bCallOnDifferentObject && !bStaticCall;
+	const bool bInterfaceCall = bCallOnDifferentObject && Statement.FunctionContext && (UEdGraphSchema_K2::PC_Interface == Statement.FunctionContext->Type.PinCategory);
 
-			for (int32 i = 0; i < ArgumentList.Num(); ++i)
+	FString Result;
+	bool bCloseCast = false;
+	if (!bInline)
+	{
+		// Handle the return value of the function being called
+		UProperty* FuncToCallReturnProperty = Statement.FunctionToCall->GetReturnProperty();
+		if (FuncToCallReturnProperty && ensure(Statement.LHS))
+		{
+			FString BeginCast;
+			if (auto ObjectProperty = Cast<UObjectProperty>(FuncToCallReturnProperty))
 			{
-				UProperty* ArgProperty = ArgumentList[i];
-
-				if (i > 0)
-				{
-					Emit(Header, TEXT(", "));
-					Emit(Body, TEXT(", "));
-				}
-
-				if (ArgProperty->HasAnyPropertyFlags(CPF_OutParm))
+				UClass* LClass = Statement.LHS ? Cast<UClass>(Statement.LHS->Type.PinSubCategoryObject.Get()) : nullptr;
+				if (LClass && LClass->IsChildOf(ObjectProperty->PropertyClass) && !ObjectProperty->PropertyClass->IsChildOf(LClass))
 				{
-					Emit(Header, TEXT("/*out*/ "));
-					Emit(Body, TEXT("/*out*/ "));
+					BeginCast = FString::Printf(TEXT("CastChecked<%s%s>("), LClass->GetPrefixCPP(), *LClass->GetName());
+					bCloseCast = true;
 				}
-
-				ArgProperty->ExportCppDeclaration(Header, EExportedDeclaration::Parameter, NULL, EPropertyExportCPPFlags::CPPF_CustomTypeName | EPropertyExportCPPFlags::CPPF_BlueprintCppBackend);
-				ArgProperty->ExportCppDeclaration(Body, EExportedDeclaration::Parameter, NULL, EPropertyExportCPPFlags::CPPF_CustomTypeName | EPropertyExportCPPFlags::CPPF_BlueprintCppBackend);
 			}
 
-			Emit(Header, TEXT(")"));
-			if (bGenerateAsNativeEventImplementation)
-			{
-				Emit(Header, TEXT(" override"));
-			}
-			Emit(Header, TEXT(";\n"));
-			Emit(Body, TEXT(")\n"));
-		}
+			Result += FString::Printf(TEXT("%s = %s"), *TermToText(EmitterContext, Statement.LHS), *BeginCast);
 
-		// Start the body of the implementation
-		Emit(Body, TEXT("{\n"));
+		}
 	}
 
-	const FString ReturnValueString = ReturnValue ? (FString(TEXT(" ")) + ReturnValue->GetName()) : TEXT("");
-
-	if (!bGenerateStubOnly)
+	// Emit object to call the method on
+	if (bInterfaceCall)
 	{
-		// Emit local variables
-		DeclareLocalVariables(FunctionContext, LocalVariables);
-
-		bool bUseSwitchState = false;
-		for (auto Node : FunctionContext.LinearExecutionList)
+		auto ContextInterfaceClass = CastChecked<UClass>(Statement.FunctionContext->Type.PinSubCategoryObject.Get());
+		ensure(ContextInterfaceClass->IsChildOf<UInterface>());
+		Result += FString::Printf(TEXT("I%s::Execute_%s(%s.GetObject(), ")
+			, *ContextInterfaceClass->GetName()
+			, *Statement.FunctionToCall->GetName()
+			, *TermToText(EmitterContext, Statement.FunctionContext, false));
+	}
+	else
+	{
+		if (bStaticCall)
 		{
-			TArray<FBlueprintCompiledStatement*>* StatementList = FunctionContext.StatementsPerNode.Find(Node);
-			if (StatementList)
-			{
-				for (auto Statement : (*StatementList))
-				{
-					if (Statement && (
-						Statement->Type == KCST_UnconditionalGoto ||
-						Statement->Type == KCST_PushState ||
-						Statement->Type == KCST_GotoIfNot ||
-						Statement->Type == KCST_ComputedGoto ||
-						Statement->Type == KCST_EndOfThread ||
-						Statement->Type == KCST_EndOfThreadIfNot ||
-						Statement->Type == KCST_GotoReturn ||
-						Statement->Type == KCST_GotoReturnIfNot))
-					{
-						bUseSwitchState = true;
-						break;
-					}
-				}
-			}
-			if (bUseSwitchState)
-			{
-				break;
-			}
+			const bool bIsCustomThunk = Statement.FunctionToCall->HasMetaData(TEXT("CustomStructureParam")) || Statement.FunctionToCall->HasMetaData(TEXT("ArrayParm"));
+			auto OwnerClass = Statement.FunctionToCall->GetOuterUClass();
+			Result += bIsCustomThunk ? TEXT("FCustomThunkTemplates::") : FString::Printf(TEXT("%s%s::"), OwnerClass->GetPrefixCPP(), *OwnerClass->GetName());
+		}
+		else if (bCallOnDifferentObject) //@TODO: Badness, could be a self reference wired to another instance!
+		{
+			Result += FString::Printf(TEXT("%s->"), *TermToText(EmitterContext, Statement.FunctionContext, false));
 		}
 
-		if (bUseSwitchState)
+		if (Statement.bIsParentContext)
+		{
+			Result += TEXT("Super::");
+		}
+		Result += Statement.FunctionToCall->GetName();
+		if (Statement.bIsParentContext && FEmitHelper::ShoulsHandleAsNativeEvent(Statement.FunctionToCall))
 		{
-			DeclareStateSwitch(FunctionContext);
+			ensure(!bCallOnDifferentObject);
+			Result += TEXT("_Implementation");
 		}
 
+		// Emit method parameter list
+		Result += TEXT("(");
+	}
+	Result += EmitMethodInputParameterList(EmitterContext, Statement);
+	Result += TEXT(")");
 
-		// Run thru code looking only at things marked as jump targets, to make sure the jump targets are ordered in order of appearance in the linear execution list
-		// Emit code in the order specified by the linear execution list (the first node is always the entry point for the function)
-		for (int32 NodeIndex = 0; NodeIndex < FunctionContext.LinearExecutionList.Num(); ++NodeIndex)
-		{
-			UEdGraphNode* StatementNode = FunctionContext.LinearExecutionList[NodeIndex];
-			TArray<FBlueprintCompiledStatement*>* StatementList = FunctionContext.StatementsPerNode.Find(StatementNode);
+	if (bCloseCast)
+	{
+		Result += TEXT(", ECastCheckedType::NullAllowed)");
+	}
 
-			if (StatementList != NULL)
-			{
-				for (int32 StatementIndex = 0; StatementIndex < StatementList->Num(); ++StatementIndex)
-				{
-					FBlueprintCompiledStatement& Statement = *((*StatementList)[StatementIndex]);
+	if (!bInline)
+	{
+		Result += TEXT(";");
+	}
 
-					if (Statement.bIsJumpTarget)
-					{
-						// Just making sure we number them in order of appearance, so jump statements don't influence the order
-						const int32 StateNum = StatementToStateIndex(FunctionContext, &Statement);
-					}
-				}
-			}
-		}
+	return Result;
+}
 
-		FEmitterLocalContext EmitterContext;
-		EmitterContext.IncreaseIndent();
-		if (bUseSwitchState)
+FString FBlueprintCompilerCppBackend::TermToText(FEmitterLocalContext& EmitterContext, const FBPTerminal* Term, bool bUseSafeContext)
+{
+	const FString PSC_Self(TEXT("self"));
+	if (Term->bIsLiteral)
+	{
+		return FEmitHelper::LiteralTerm(EmitterContext, Term->Type, Term->Name, Term->ObjectLiteral);
+	}
+	else if (Term->InlineGeneratedParameter)
+	{
+		if (KCST_SwitchValue == Term->InlineGeneratedParameter->Type)
 		{
-			EmitterContext.IncreaseIndent();
-			EmitterContext.IncreaseIndent();
-			EmitterContext.IncreaseIndent();
+			return EmitSwitchValueStatmentInner(EmitterContext, *Term->InlineGeneratedParameter);
 		}
-		EmitterContext.SetCurrentlyGeneratedClass(FunctionContext.NewClass);
-
-		// Emit code in the order specified by the linear execution list (the first node is always the entry point for the function)
-		for (int32 NodeIndex = 0; NodeIndex < FunctionContext.LinearExecutionList.Num(); ++NodeIndex)
+		else if (KCST_CallFunction == Term->InlineGeneratedParameter->Type)
+		{
+			return EmitCallStatmentInner(EmitterContext, *Term->InlineGeneratedParameter, true);
+		}
+		else
+		{
+			ensureMsgf(false, TEXT("KCST %d is not accepted as inline statement."), Term->InlineGeneratedParameter->Type);
+			return FString();
+		}
+	}
+	else
+	{
+		FString ResultPath(TEXT(""));
+		if ((Term->Context != NULL) && (Term->Context->Name != PSC_Self))
 		{
-			UEdGraphNode* StatementNode = FunctionContext.LinearExecutionList[NodeIndex];
-			TArray<FBlueprintCompiledStatement*>* StatementList = FunctionContext.StatementsPerNode.Find(StatementNode);
+			ResultPath = TermToText(EmitterContext, Term->Context, false);
 
-			if (StatementList != NULL)
+			if (Term->Context->IsStructContextType())
 			{
-				for (int32 StatementIndex = 0; StatementIndex < StatementList->Num(); ++StatementIndex)
-				{
-					FBlueprintCompiledStatement& Statement = *((*StatementList)[StatementIndex]);
-
-					if (Statement.bIsJumpTarget && bUseSwitchState)
-					{
-						EmitterContext.DecreaseIndent();
-						EmitterContext.AddLine(TEXT("}"));
-						const int32 StateNum = StatementToStateIndex(FunctionContext, &Statement);
-						EmitterContext.DecreaseIndent();
-						EmitterContext.AddLine(FString::Printf(TEXT("case %d:"), StateNum));
-						EmitterContext.IncreaseIndent();
-						EmitterContext.AddLine(TEXT("{"));
-						EmitterContext.IncreaseIndent();
-					}
-
-					switch (Statement.Type)
-					{
-					case KCST_Nop:
-						EmitterContext.AddLine(TEXT("//No operation."));
-						break;
-					case KCST_CallFunction:
-						EmitCallStatment(EmitterContext, FunctionContext, Statement);
-						break;
-					case KCST_Assignment:
-						EmitAssignmentStatment(EmitterContext, FunctionContext, Statement);
-						break;
-					case KCST_CompileError:
-						UE_LOG(LogK2Compiler, Error, TEXT("C++ backend encountered KCST_CompileError"));
-						EmitterContext.AddLine(TEXT("static_assert(false); // KCST_CompileError"));
-						break;
-					case KCST_PushState:
-						EmitPushStateStatement(EmitterContext, FunctionContext, Statement);
-						break;
-					case KCST_Return:
-						EmitterContext.AddLine(TEXT("// Return statement."));
-						break;
-					case KCST_EndOfThread:
-						EmitEndOfThreadStatement(EmitterContext, FunctionContext, ReturnValueString);
-						break;
-					case KCST_Comment:
-						EmitterContext.AddLine(FString::Printf(TEXT("// %s"), *Statement.Comment));
-						break;
-					case KCST_DebugSite:
-						break;
-					case KCST_CastObjToInterface:
-						EmitCastObjToInterfaceStatement(EmitterContext, FunctionContext, Statement);
-						break;
-					case KCST_DynamicCast:
-						EmitDynamicCastStatement(EmitterContext, FunctionContext, Statement);
-						break;
-					case KCST_ObjectToBool:
-						EmitObjectToBoolStatement(EmitterContext, FunctionContext, Statement);
-						break;
-					case KCST_AddMulticastDelegate:
-						EmitAddMulticastDelegateStatement(EmitterContext, FunctionContext, Statement);
-						break;
-					case KCST_ClearMulticastDelegate:
-						EmitClearMulticastDelegateStatement(EmitterContext, FunctionContext, Statement);
-						break;
-					case KCST_WireTraceSite:
-						break;
-					case KCST_BindDelegate:
-						EmitBindDelegateStatement(EmitterContext, FunctionContext, Statement);
-						break;
-					case KCST_RemoveMulticastDelegate:
-						EmitRemoveMulticastDelegateStatement(EmitterContext, FunctionContext, Statement);
-						break;
-					case KCST_CallDelegate:
-						EmitCallDelegateStatment(EmitterContext, FunctionContext, Statement);
-						break;
-					case KCST_CreateArray:
-						EmitCreateArrayStatement(EmitterContext, FunctionContext, Statement);
-						break;
-					case KCST_CrossInterfaceCast:
-						EmitCastBetweenInterfacesStatement(EmitterContext, FunctionContext, Statement);
-						break;
-					case KCST_MetaCast:
-						EmitMetaCastStatement(EmitterContext, FunctionContext, Statement);
-						break;
-					case KCST_CastInterfaceToObj:
-						EmitCastInterfaceToObjStatement(EmitterContext, FunctionContext, Statement);
-						break;
-					case KCST_ComputedGoto:
-					case KCST_UnconditionalGoto:
-					case KCST_GotoIfNot:
-					case KCST_EndOfThreadIfNot:
-					case KCST_GotoReturn:
-					case KCST_GotoReturnIfNot:
-						EmitGotoStatement(EmitterContext, FunctionContext, Statement);
-						break;
-					case KCST_SwitchValue:
-						// Switch Value should be always an "inline" statement, so there is no point to handle it here
-						// case: KCST_AssignmentOnPersistentFrame
-					default:
-						EmitterContext.AddLine(TEXT("// Warning: Ignoring unsupported statement\n"));
-						UE_LOG(LogK2Compiler, Error, TEXT("C++ backend encountered unsupported statement type %d"), (int32)Statement.Type);
-						break;
-					};
-				}
+				ResultPath += TEXT(".");
+			}
+			else
+			{
+				ResultPath += TEXT("->");
 			}
 		}
 
-		Emit(Body, *EmitterContext.GetResult());
-		if (bUseSwitchState)
+		FString Conditions;
+		if (bUseSafeContext)
 		{
-			CloseStateSwitch(FunctionContext);
+			Conditions = FSafeContextScopedEmmitter::ValidationChain(EmitterContext, Term->Context, *this);
 		}
-	}
-
-	EmitReturnStatement(FunctionContext, ReturnValueString);
 
-	Emit(Body, TEXT("}\n\n"));
+		ResultPath += Term->AssociatedVarProperty ? Term->AssociatedVarProperty->GetNameCPP() : Term->Name;
+		return Conditions.IsEmpty()
+			? ResultPath
+			: FString::Printf(TEXT("((%s) ? (%s) : (%s))"), *Conditions, *ResultPath, *FEmitHelper::DefaultValue(EmitterContext, Term->Type));
+	}
 }
 
-void FBlueprintCompilerCppBackend::GenerateCodeFromEnum(UUserDefinedEnum* SourceEnum)
+FString FBlueprintCompilerCppBackend::LatentFunctionInfoTermToText(FEmitterLocalContext& EmitterContext, FBPTerminal* Term, FBlueprintCompiledStatement* TargetLabel)
 {
-	check(SourceEnum);
-	const FString Name = SourceEnum->GetName();
-	EmitFileBeginning(Name, nullptr);
-
-	Emit(Header, TEXT("UENUM(BlueprintType)\nenum class "));
-	Emit(Header, *Name);
-	Emit(Header, TEXT(" : uint8\n{"));
+	auto LatentInfoStruct = FLatentActionInfo::StaticStruct();
 
-	for (int32 Index = 0; Index < SourceEnum->NumEnums(); ++Index)
+	// Find the term name we need to fixup
+	FString FixupTermName;
+	for (UProperty* Prop = LatentInfoStruct->PropertyLink; Prop; Prop = Prop->PropertyLinkNext)
 	{
-		const FString ElemName = SourceEnum->GetEnumName(Index);
-		const int32 ElemValue = Index;
-
-		const FString& DisplayNameMD = SourceEnum->GetMetaData(TEXT("DisplayName"), ElemValue);// TODO: value or index?
-		const FString Meta = DisplayNameMD.IsEmpty() ? FString() : FString::Printf(TEXT("UMETA(DisplayName = \"%s\")"), *DisplayNameMD);
-		Emit(Header, *FString::Printf(TEXT("\n\t%s = %d %s,"), *ElemName, ElemValue, *Meta));
+		static const FName NeedsLatentFixup(TEXT("NeedsLatentFixup"));
+		if (Prop->GetBoolMetaData(NeedsLatentFixup))
+		{
+			FixupTermName = Prop->GetName();
+			break;
+		}
 	}
-	Emit(Header, TEXT("\n};\n"));
-}
 
-void FBlueprintCompilerCppBackend::GenerateCodeFromStruct(UUserDefinedStruct* SourceStruct)
-{
-	check(SourceStruct);
+	check(!FixupTermName.IsEmpty());
 
-	EmitFileBeginning(SourceStruct->GetName(), SourceStruct);
+	FString StructValues = Term->Name;
 
-	const FString NewName = FString(TEXT("F")) + SourceStruct->GetName();
-	Emit(Header, TEXT("USTRUCT(BlueprintType)\n"));
+	// Index 0 is always the ubergraph
+	const int32 TargetStateIndex = StateMapPerFunction[0].StatementToStateIndex(TargetLabel);
+	const int32 LinkageTermStartIdx = StructValues.Find(FixupTermName);
+	check(LinkageTermStartIdx != INDEX_NONE);
+	StructValues = StructValues.Replace(TEXT("-1"), *FString::FromInt(TargetStateIndex));
 
-	Emit(Header, *FString::Printf(TEXT("struct %s\n{\npublic:\n\tGENERATED_BODY()\n"), *NewName));
-	EmitStructProperties(Header, SourceStruct);
-	Emit(Header, *FEmitDefaultValueHelper::GenerateGetDefaultValue(SourceStruct));
-	Emit(Header, TEXT("};\n"));
+	return FEmitHelper::LiteralTerm(EmitterContext, Term->Type, StructValues, nullptr);
 }
 
-void FBlueprintCompilerCppBackend::EmitFileBeginning(const FString& CleanName, UStruct* SourceStruct)
+FString FBlueprintCompilerCppBackend::InnerFunctionImplementation(FKismetFunctionContext& FunctionContext, bool bUseSwitchState)
 {
-	Emit(Header, TEXT("#pragma once\n\n"));
-	if (SourceStruct)
-	{ 
+	FEmitterLocalContext EmitterContext;
+	EmitterContext.IncreaseIndent();
+	if (bUseSwitchState)
+	{
+		if (FunctionContext.bUseFlowStack)
 		{
-			TArray<FString> PersistentHeaders;
-			PersistentHeaders.Add(FString(FApp::GetGameName()) + TEXT(".h"));
-			PersistentHeaders.Add(CleanName + TEXT(".h"));
-			PersistentHeaders.Add(TEXT("GeneratedCodeHelpers.h"));
-			Emit(Body, *FEmitHelper::GatherNativeHeadersToInclude(SourceStruct, PersistentHeaders));
+			EmitterContext.AddLine(TEXT("TArray< int32, TInlineAllocator<8> > StateStack;\n"));
 		}
 
-		// find objects referenced by functions/script
-		TArray<UObject*> IncludeInHeader;
-		TArray<UObject*> IncludeInBody;
+		EmitterContext.AddLine(TEXT("int32 CurrentState = 0;"));
+		EmitterContext.AddLine(TEXT("do"));
+		EmitterContext.AddLine(TEXT("{"));
+		EmitterContext.IncreaseIndent();
+		EmitterContext.AddLine(TEXT("switch( CurrentState )"));
+		EmitterContext.AddLine(TEXT("{"));
+		EmitterContext.AddLine(TEXT("case 0:"));
+		EmitterContext.IncreaseIndent();
+		EmitterContext.AddLine(TEXT("{"));
+		EmitterContext.IncreaseIndent();
+	}
+	EmitterContext.SetCurrentlyGeneratedClass(FunctionContext.NewClass);
+
+	// Emit code in the order specified by the linear execution list (the first node is always the entry point for the function)
+	for (int32 NodeIndex = 0; NodeIndex < FunctionContext.LinearExecutionList.Num(); ++NodeIndex)
+	{
+		UEdGraphNode* StatementNode = FunctionContext.LinearExecutionList[NodeIndex];
+		TArray<FBlueprintCompiledStatement*>* StatementList = FunctionContext.StatementsPerNode.Find(StatementNode);
+
+		if (!StatementList)
 		{
-			FReferenceFinder HeaderReferenceFinder(IncludeInHeader, NULL, false, false, true, false);
-			FReferenceFinder BodyReferenceFinder(IncludeInBody, NULL, false, false, true, false);
-			{
-				TArray<UObject*> ObjectsToCheck;
-				GetObjectsWithOuter(SourceStruct, ObjectsToCheck, true);
-				for (auto Obj : ObjectsToCheck)
-				{
-					auto Property = Cast<UProperty>(Obj);
-					auto OwnerProperty = IsValid(Property) ? Property->GetOwnerProperty() : nullptr;
-					const bool bIsParam = OwnerProperty && (0 != (OwnerProperty->PropertyFlags & CPF_Parm));
-					const bool bIsMemberVariable = OwnerProperty && (OwnerProperty->GetOuter() == SourceStruct);
-					if (bIsParam || bIsMemberVariable)
-					{
-						HeaderReferenceFinder.FindReferences(Obj);
-					}
-					else
-					{
-						BodyReferenceFinder.FindReferences(Obj);
-					}
-				}
-			}
+			continue;
+		}
+
+		for (int32 StatementIndex = 0; StatementIndex < StatementList->Num(); ++StatementIndex)
+		{
+			FBlueprintCompiledStatement& Statement = *((*StatementList)[StatementIndex]);
 
-			if (auto SuperStruct = SourceStruct->GetSuperStruct())
+			if (Statement.bIsJumpTarget && bUseSwitchState)
 			{
-				IncludeInHeader.AddUnique(SuperStruct);
+				EmitterContext.DecreaseIndent();
+				EmitterContext.AddLine(TEXT("}"));
+				const int32 StateNum = StatementToStateIndex(FunctionContext, &Statement);
+				EmitterContext.DecreaseIndent();
+				EmitterContext.AddLine(FString::Printf(TEXT("case %d:"), StateNum));
+				EmitterContext.IncreaseIndent();
+				EmitterContext.AddLine(TEXT("{"));
+				EmitterContext.IncreaseIndent();
 			}
 
-			if (auto SourceClass = Cast<UClass>(SourceStruct))
+			switch (Statement.Type)
 			{
-				for (auto& ImplementedInterface : SourceClass->Interfaces)
-				{
-					IncludeInHeader.AddUnique(ImplementedInterface.Class);
-				}
-			}
+			case KCST_Nop:
+				EmitterContext.AddLine(TEXT("//No operation."));
+				break;
+			case KCST_CallFunction:
+				EmitCallStatment(EmitterContext, FunctionContext, Statement);
+				break;
+			case KCST_Assignment:
+				EmitAssignmentStatment(EmitterContext, FunctionContext, Statement);
+				break;
+			case KCST_CompileError:
+				UE_LOG(LogK2Compiler, Error, TEXT("C++ backend encountered KCST_CompileError"));
+				EmitterContext.AddLine(TEXT("static_assert(false); // KCST_CompileError"));
+				break;
+			case KCST_PushState:
+				EmitPushStateStatement(EmitterContext, FunctionContext, Statement);
+				break;
+			case KCST_Return:
+				UE_LOG(LogK2Compiler, Error, TEXT("C++ backend encountered KCST_Return"));
+				EmitterContext.AddLine(TEXT("// Return statement."));
+				break;
+			case KCST_EndOfThread:
+				EmitEndOfThreadStatement(EmitterContext, FunctionContext);
+				break;
+			case KCST_Comment:
+				EmitterContext.AddLine(FString::Printf(TEXT("// %s"), *Statement.Comment));
+				break;
+			case KCST_DebugSite:
+				break;
+			case KCST_CastObjToInterface:
+				EmitCastObjToInterfaceStatement(EmitterContext, FunctionContext, Statement);
+				break;
+			case KCST_DynamicCast:
+				EmitDynamicCastStatement(EmitterContext, FunctionContext, Statement);
+				break;
+			case KCST_ObjectToBool:
+				EmitObjectToBoolStatement(EmitterContext, FunctionContext, Statement);
+				break;
+			case KCST_AddMulticastDelegate:
+				EmitAddMulticastDelegateStatement(EmitterContext, FunctionContext, Statement);
+				break;
+			case KCST_ClearMulticastDelegate:
+				EmitClearMulticastDelegateStatement(EmitterContext, FunctionContext, Statement);
+				break;
+			case KCST_WireTraceSite:
+				break;
+			case KCST_BindDelegate:
+				EmitBindDelegateStatement(EmitterContext, FunctionContext, Statement);
+				break;
+			case KCST_RemoveMulticastDelegate:
+				EmitRemoveMulticastDelegateStatement(EmitterContext, FunctionContext, Statement);
+				break;
+			case KCST_CallDelegate:
+				EmitCallDelegateStatment(EmitterContext, FunctionContext, Statement);
+				break;
+			case KCST_CreateArray:
+				EmitCreateArrayStatement(EmitterContext, FunctionContext, Statement);
+				break;
+			case KCST_CrossInterfaceCast:
+				EmitCastBetweenInterfacesStatement(EmitterContext, FunctionContext, Statement);
+				break;
+			case KCST_MetaCast:
+				EmitMetaCastStatement(EmitterContext, FunctionContext, Statement);
+				break;
+			case KCST_CastInterfaceToObj:
+				EmitCastInterfaceToObjStatement(EmitterContext, FunctionContext, Statement);
+				break;
+			case KCST_ComputedGoto:
+			case KCST_UnconditionalGoto:
+			case KCST_GotoIfNot:
+			case KCST_EndOfThreadIfNot:
+			case KCST_GotoReturn:
+			case KCST_GotoReturnIfNot:
+				EmitGotoStatement(EmitterContext, FunctionContext, Statement);
+				break;
+			case KCST_SwitchValue:
+				// Switch Value should be always an "inline" statement, so there is no point to handle it here
+				// case: KCST_AssignmentOnPersistentFrame
+			default:
+				EmitterContext.AddLine(TEXT("// Warning: Ignoring unsupported statement\n"));
+				UE_LOG(LogK2Compiler, Error, TEXT("C++ backend encountered unsupported statement type %d"), (int32)Statement.Type);
+				break;
+			};
 		}
+	}
 
-		TSet<FString> AlreadyIncluded;
-		AlreadyIncluded.Add(SourceStruct->GetName());
-
-		auto EmitInner = [&](FStringOutputDevice& Dst, const TArray<UObject*>& Src)
+	if (bUseSwitchState)
+	{
+		EmitterContext.DecreaseIndent();
+		// Default error-catching case 
+		EmitterContext.AddLine(TEXT("}"));
+		EmitterContext.DecreaseIndent();
+		EmitterContext.AddLine(TEXT("default:"));
+		EmitterContext.IncreaseIndent();
+		if (FunctionContext.bUseFlowStack)
 		{
-			for (UObject* Obj : Src)
-			{
-				bool bWantedType = Obj && (Obj->IsA<UBlueprintGeneratedClass>() || Obj->IsA<UUserDefinedEnum>() || Obj->IsA<UUserDefinedStruct>());
-				if (!bWantedType)
-				{
-					for (UObject* OuterObj = (Obj ? Obj->GetOuter() : nullptr); nullptr != OuterObj; OuterObj = OuterObj->GetOuter())
-					{
-						if (OuterObj->IsA<UBlueprintGeneratedClass>() || OuterObj->IsA<UUserDefinedStruct>())
-						{
-							Obj = OuterObj;
-							bWantedType = true;
-							break;
-						}
-					}
-				}
-
-				if (!bWantedType)
-				{
-					Obj = Obj ? Cast<UBlueprintGeneratedClass>(Obj->GetClass()) : nullptr;
-					bWantedType = (nullptr != Obj);
-				}
+			EmitterContext.AddLine(TEXT("check(false); // Invalid state"));
+		}
+		EmitterContext.AddLine(TEXT("break;"));
+		EmitterContext.DecreaseIndent();
+		// Close the switch block and do-while loop
+		EmitterContext.AddLine(TEXT("}"));
+		EmitterContext.DecreaseIndent();
+		EmitterContext.AddLine(TEXT("} while( CurrentState != -1 );"));
+	}
 
-				if (bWantedType && Obj)
-				{
-					const FString Name = Obj->GetName();
-					if (!AlreadyIncluded.Contains(Name))
-					{
-						AlreadyIncluded.Add(Name);
-						Emit(Dst, *FString::Printf(TEXT("#include \"%s.h\"\n"), *Name));
-					}
-				}
-			}
-			Emit(Dst, TEXT("\n"));
-		};
 
-		EmitInner(Header, IncludeInHeader);
-		EmitInner(Body, IncludeInBody);
-	}
-	Emit(Header, *FString::Printf(TEXT("#include \"%s.generated.h\"\n\n"), *CleanName));
+	return EmitterContext.GetResult();
 }
diff --git a/Engine/Source/Developer/BlueprintCompilerCppBackend/Private/BlueprintCompilerCppBackend.h b/Engine/Source/Developer/BlueprintCompilerCppBackend/Private/BlueprintCompilerCppBackend.h
--- a/Engine/Source/Developer/BlueprintCompilerCppBackend/Private/BlueprintCompilerCppBackend.h
+++ b/Engine/Source/Developer/BlueprintCompilerCppBackend/Private/BlueprintCompilerCppBackend.h
@@ -1,62 +1,16 @@
 // Copyright 1998-2015 Epic Games, Inc. All Rights Reserved.
 #pragma  once
 
-#include "IBlueprintCompilerCppBackendModule.h"
+#include "BlueprintCompilerCppBackendBase.h"
 
 struct FEmitterLocalContext;
 
-class FBlueprintCompilerCppBackend : public IBlueprintCompilerCppBackend
+class FBlueprintCompilerCppBackend : public FBlueprintCompilerCppBackendBase
 {
 protected:
-	struct FFunctionLabelInfo
-	{
-		TMap<FBlueprintCompiledStatement*, int32> StateMap;
-		int32 StateCounter;
-
-		FFunctionLabelInfo()
-		{
-			StateCounter = 0;
-		}
-
-		int32 StatementToStateIndex(FBlueprintCompiledStatement* Statement)
-		{
-			int32& Index = StateMap.FindOrAdd(Statement);
-			if (Index == 0)
-			{
-				Index = ++StateCounter;
-			}
-
-			return Index;
-		}
-	};
-
-	TArray<FFunctionLabelInfo> StateMapPerFunction;
-	TMap<FKismetFunctionContext*, int32> FunctionIndexMap;
-
-	FString CppClassName;
-
-public:
-	FStringOutputDevice Header;
-	FStringOutputDevice Body;
-
-	// IBlueprintCompilerCppBackend implementation
-	virtual void GenerateCodeFromClass(UClass* SourceClass, TIndirectArray<FKismetFunctionContext>& Functions, bool bGenerateStubsOnly) override;
-	virtual void GenerateCodeFromEnum(UUserDefinedEnum* SourceEnum) override;
-	virtual void GenerateCodeFromStruct(UUserDefinedStruct* SourceStruct) override;
-
-	virtual const FString& GetBody()	const override { return Body; }
-	virtual const FString& GetHeader()	const override { return Header; }
-	// end of IBlueprintCompilerCppBackend implementation
-
-public:
-
-	void Emit(FStringOutputDevice& Target, const TCHAR* Message)
-	{
-		Target += Message;
-	}
-
-	void EmitStructProperties(FStringOutputDevice& Target, UStruct* SourceClass);
+	virtual FString InnerFunctionImplementation(FKismetFunctionContext& FunctionContext, bool bUseSwitchState) override;
 
+protected:
 	void EmitCallStatment(FEmitterLocalContext& EmitterContext, FKismetFunctionContext& FunctionContext, FBlueprintCompiledStatement& Statement);
 	void EmitCallDelegateStatment(FEmitterLocalContext& EmitterContext, FKismetFunctionContext& FunctionContext, FBlueprintCompiledStatement& Statement);
 	void EmitAssignmentStatment(FEmitterLocalContext& EmitterContext, FKismetFunctionContext& FunctionContext, FBlueprintCompiledStatement& Statement);
@@ -71,38 +25,16 @@ public:
 	void EmitCreateArrayStatement(FEmitterLocalContext& EmitterContext, FKismetFunctionContext& FunctionContext, FBlueprintCompiledStatement& Statement);
 	void EmitGotoStatement(FEmitterLocalContext& EmitterContext, FKismetFunctionContext& FunctionContext, FBlueprintCompiledStatement& Statement);
 	void EmitPushStateStatement(FEmitterLocalContext& EmitterContext, FKismetFunctionContext& FunctionContext, FBlueprintCompiledStatement& Statement);
-	void EmitEndOfThreadStatement(FEmitterLocalContext& EmitterContext, FKismetFunctionContext& FunctionContext, const FString& ReturnValueString);
+	void EmitEndOfThreadStatement(FEmitterLocalContext& EmitterContext, FKismetFunctionContext& FunctionContext);
 	void EmitRemoveMulticastDelegateStatement(FEmitterLocalContext& EmitterContext, FKismetFunctionContext& FunctionContext, FBlueprintCompiledStatement& Statement);
 	void EmitMetaCastStatement(FEmitterLocalContext& EmitterContext, FKismetFunctionContext& FunctionContext, FBlueprintCompiledStatement& Statement);
 
-	void EmitReturnStatement(FKismetFunctionContext& FunctionContext, const FString& ReturnValueString);
-
-	/** Emits local variable declarations for a function */
-	void DeclareLocalVariables(FKismetFunctionContext& FunctionContext, TArray<UProperty*>& LocalVariables);
-	
-	/** Emits the internal execution flow state declaration for a function */
-	void DeclareStateSwitch(FKismetFunctionContext& FunctionContext);
-	void CloseStateSwitch(FKismetFunctionContext& FunctionContext);
-
-	/** Builds both the header declaration and body implementation of a function */
-	void ConstructFunction(FKismetFunctionContext& FunctionContext, bool bGenerateStubOnly);
-
-	FString TermToText(FEmitterLocalContext& EmitterContext, const FBPTerminal* Term, bool bUseSafeContext = true);
-
 protected:
-	void EmitFileBeginning(const FString& CleanName, UStruct* SourceStruct);
-
-	int32 StatementToStateIndex(FKismetFunctionContext& FunctionContext, FBlueprintCompiledStatement* Statement)
-	{
-		int32 Index = FunctionIndexMap.FindChecked(&FunctionContext);
-		return StateMapPerFunction[Index].StatementToStateIndex(Statement);
-	}
-
 	FString LatentFunctionInfoTermToText(FEmitterLocalContext& EmitterContext, FBPTerminal* Term, FBlueprintCompiledStatement* TargetLabel);
-	
 	FString EmitMethodInputParameterList(FEmitterLocalContext& EmitterContext, FBlueprintCompiledStatement& Statement);
-
 	FString EmitSwitchValueStatmentInner(FEmitterLocalContext& EmitterContext, FBlueprintCompiledStatement& Statement);
-
 	FString EmitCallStatmentInner(FEmitterLocalContext& EmitterContext, FBlueprintCompiledStatement& Statement, bool bInline);
+
+public:
+	FString TermToText(FEmitterLocalContext& EmitterContext, const FBPTerminal* Term, bool bUseSafeContext = true);
 };
\ No newline at end of file
diff --git a/Engine/Source/Developer/BlueprintCompilerCppBackend/Private/BlueprintCompilerCppBackendBase.cpp b/Engine/Source/Developer/BlueprintCompilerCppBackend/Private/BlueprintCompilerCppBackendBase.cpp
new file mode 100644
--- /dev/null
+++ b/Engine/Source/Developer/BlueprintCompilerCppBackend/Private/BlueprintCompilerCppBackendBase.cpp
@@ -0,0 +1,485 @@
+// Copyright 1998-2015 Epic Games, Inc. All Rights Reserved.
+
+#include "BlueprintCompilerCppBackendModulePrivatePCH.h"
+#include "BlueprintCompilerCppBackendBase.h"
+#include "BlueprintCompilerCppBackendUtils.h"
+
+void FBlueprintCompilerCppBackendBase::EmitStructProperties(FStringOutputDevice& Target, UStruct* SourceClass)
+{
+	// Emit class variables
+	for (TFieldIterator<UProperty> It(SourceClass, EFieldIteratorFlags::ExcludeSuper); It; ++It)
+	{
+		UProperty* Property = *It;
+		check(Property);
+
+		Emit(Header, TEXT("\n\tUPROPERTY("));
+		{
+			TArray<FString> Tags = FEmitHelper::ProperyFlagsToTags(Property->PropertyFlags);
+			Tags.Emplace(FEmitHelper::HandleRepNotifyFunc(Property));
+			Tags.Emplace(FEmitHelper::HandleMetaData(Property, false));
+			Tags.Remove(FString());
+
+			FString AllTags;
+			FEmitHelper::ArrayToString(Tags, AllTags, TEXT(", "));
+			Emit(Header, *AllTags);
+		}
+		Emit(Header, TEXT(")\n"));
+		Emit(Header, TEXT("\t"));
+		Property->ExportCppDeclaration(Target, EExportedDeclaration::Member, NULL, EPropertyExportCPPFlags::CPPF_CustomTypeName | EPropertyExportCPPFlags::CPPF_BlueprintCppBackend);
+		Emit(Header, TEXT(";\n"));
+	}
+}
+
+void FBlueprintCompilerCppBackendBase::GenerateCodeFromClass(UClass* SourceClass, TIndirectArray<FKismetFunctionContext>& Functions, bool bGenerateStubsOnly)
+{
+	auto CleanCppClassName = SourceClass->GetName();
+	CppClassName = FString(SourceClass->GetPrefixCPP()) + CleanCppClassName;
+	
+	EmitFileBeginning(CleanCppClassName, SourceClass);
+
+	// MC DELEGATE DECLARATION
+	{
+		auto DelegateDeclarations = FEmitHelper::EmitMulticastDelegateDeclarations(SourceClass);
+		FString AllDeclarations;
+		FEmitHelper::ArrayToString(DelegateDeclarations, AllDeclarations, TEXT(";\n"));
+		if (DelegateDeclarations.Num())
+		{
+			Emit(Header, *AllDeclarations);
+			Emit(Header, TEXT(";\n"));
+		}
+	}
+
+	// GATHER ALL SC DELEGATES
+	{
+		TArray<UDelegateProperty*> Delegates;
+		for (TFieldIterator<UDelegateProperty> It(SourceClass, EFieldIteratorFlags::ExcludeSuper); It; ++It)
+		{
+			Delegates.Add(*It);
+		}
+
+		for (auto& FuncContext : Functions)
+		{
+			for (TFieldIterator<UDelegateProperty> It(FuncContext.Function, EFieldIteratorFlags::ExcludeSuper); It; ++It)
+			{
+				Delegates.Add(*It);
+			}
+		}
+
+		auto DelegateDeclarations = FEmitHelper::EmitSinglecastDelegateDeclarations(Delegates);
+		FString AllDeclarations;
+		FEmitHelper::ArrayToString(DelegateDeclarations, AllDeclarations, TEXT(";\n"));
+		if (DelegateDeclarations.Num())
+		{
+			Emit(Header, *AllDeclarations);
+			Emit(Header, TEXT(";\n"));
+		}
+	}
+
+	// Class declaration
+	const bool bIsInterface = SourceClass->IsChildOf<UInterface>();
+	if (bIsInterface)
+	{
+		Emit(Header, *FString::Printf(TEXT("UINTERFACE(Blueprintable)\nclass U%s : public UInterface\n{\n\tGENERATED_BODY()\n};\n"), *CleanCppClassName));
+		Emit(Header, *FString::Printf(TEXT("\nclass I%s"), *CleanCppClassName));
+	}
+	else
+	{
+		Emit(Header, TEXT("UCLASS("));
+		if (!SourceClass->IsChildOf<UBlueprintFunctionLibrary>())
+		{
+			Emit(Header, TEXT("Blueprintable, BlueprintType"));
+		}
+		Emit(Header, TEXT(")\n"));
+
+		UClass* SuperClass = SourceClass->GetSuperClass();
+		Emit(Header, *FString::Printf(TEXT("class %s : public %s%s"), *CppClassName, SuperClass->GetPrefixCPP(), *SuperClass->GetName()));
+
+		for (auto& ImplementedInterface : SourceClass->Interfaces)
+		{
+			if (ImplementedInterface.Class)
+			{
+				Emit(Header, *FString::Printf(TEXT(", public I%s"), *ImplementedInterface.Class->GetName()));
+			}
+		}
+	}
+
+	// Begin scope
+	Emit(Header,TEXT("\n{\npublic:\n\tGENERATED_BODY()\n"));
+
+	EmitStructProperties(Header, SourceClass);
+
+	// Create the state map
+	for (int32 i = 0; i < Functions.Num(); ++i)
+	{
+		StateMapPerFunction.Add(FFunctionLabelInfo());
+		FunctionIndexMap.Add(&Functions[i], i);
+	}
+
+	// Emit function declarations and definitions (writes to header and body simultaneously)
+	if (Functions.Num() > 0)
+	{
+		Emit(Header, TEXT("\n"));
+	}
+
+	if (!bIsInterface)
+	{
+		Emit(Header, *FString::Printf(TEXT("\t%s(const FObjectInitializer& ObjectInitializer = FObjectInitializer::Get());\n\n"), *CppClassName));
+		Emit(Body, *FEmitDefaultValueHelper::GenerateConstructor(SourceClass));
+	}
+
+	for (int32 i = 0; i < Functions.Num(); ++i)
+	{
+		if (Functions[i].IsValid())
+		{
+			ConstructFunction(Functions[i], bGenerateStubsOnly);
+		}
+	}
+
+	Emit(Header, TEXT("};\n\n"));
+
+	Emit(Body, *FEmitHelper::EmitLifetimeReplicatedPropsImpl(SourceClass, CppClassName, TEXT("")));
+}
+
+void FBlueprintCompilerCppBackendBase::DeclareLocalVariables(FKismetFunctionContext& FunctionContext, TArray<UProperty*>& LocalVariables)
+{
+	for (int32 i = 0; i < LocalVariables.Num(); ++i)
+	{
+		UProperty* LocalVariable = LocalVariables[i];
+
+		Emit(Body, TEXT("\t"));
+		LocalVariable->ExportCppDeclaration(Body, EExportedDeclaration::Local, NULL, EPropertyExportCPPFlags::CPPF_CustomTypeName | EPropertyExportCPPFlags::CPPF_BlueprintCppBackend);
+		Emit(Body, TEXT("{};\n"));
+	}
+
+	if (LocalVariables.Num() > 0)
+	{
+		Emit(Body, TEXT("\n"));
+	}
+}
+
+void FBlueprintCompilerCppBackendBase::ConstructFunction(FKismetFunctionContext& FunctionContext, bool bGenerateStubOnly)
+{
+	if (FunctionContext.IsDelegateSignature())
+	{
+		return;
+	}
+
+	UFunction* Function = FunctionContext.Function;
+
+	UProperty* ReturnValue = NULL;
+	TArray<UProperty*> LocalVariables;
+
+	{
+		FString FunctionName;
+		Function->GetName(FunctionName);
+
+		TArray<UProperty*> ArgumentList;
+
+		// Split the function property list into arguments, a return value (if any), and local variable declarations
+		for (TFieldIterator<UProperty> It(Function); It; ++It)
+		{
+			UProperty* Property = *It;
+			if (Property->HasAnyPropertyFlags(CPF_Parm))
+			{
+				if (Property->HasAnyPropertyFlags(CPF_ReturnParm))
+				{
+					if (ReturnValue == NULL)
+					{
+						ReturnValue = Property;
+						LocalVariables.Add(Property);
+					}
+					else
+					{
+						UE_LOG(LogK2Compiler, Error, TEXT("Function %s from graph @@ has more than one return value (named %s and %s)"),
+							*FunctionName, *GetPathNameSafe(FunctionContext.SourceGraph), *ReturnValue->GetName(), *Property->GetName());
+					}
+				}
+				else
+				{
+					ArgumentList.Add(Property);
+				}
+			}
+			else
+			{
+				LocalVariables.Add(Property);
+			}
+		}
+
+		bool bAddCppFromBpEventMD = false;
+		bool bGenerateAsNativeEventImplementation = false;
+		// Get original function declaration
+		if (FEmitHelper::ShoulsHandleAsNativeEvent(Function)) // BlueprintNativeEvent
+		{
+			bGenerateAsNativeEventImplementation = true;
+			FunctionName += TEXT("_Implementation");
+		}
+		else if (FEmitHelper::ShoulsHandleAsImplementableEvent(Function)) // BlueprintImplementableEvent
+		{
+			bAddCppFromBpEventMD = true;
+		}
+
+		// Emit the declaration
+		const FString ReturnType = ReturnValue ? ReturnValue->GetCPPType(NULL, EPropertyExportCPPFlags::CPPF_CustomTypeName) : TEXT("void");
+
+		//@TODO: Make the header+body export more uniform
+		{
+			const FString Start = FString::Printf(TEXT("%s %s%s%s("), *ReturnType, TEXT("%s"), TEXT("%s"), *FunctionName);
+
+			TArray<FString> AdditionalMetaData;
+			if (bAddCppFromBpEventMD)
+			{
+				AdditionalMetaData.Emplace(TEXT("CppFromBpEvent"));
+			}
+
+			if (!bGenerateAsNativeEventImplementation)
+			{
+				Emit(Header, *FString::Printf(TEXT("\t%s\n"), *FEmitHelper::EmitUFuntion(Function, &AdditionalMetaData)));
+			}
+			Emit(Header, TEXT("\t"));
+
+			if (Function->HasAllFunctionFlags(FUNC_Static))
+			{
+				Emit(Header, TEXT("static "));
+			}
+
+			if (bGenerateAsNativeEventImplementation)
+			{
+				Emit(Header, TEXT("virtual "));
+			}
+
+			Emit(Header, *FString::Printf(*Start, TEXT(""), TEXT("")));
+			Emit(Body, *FString::Printf(*Start, *CppClassName, TEXT("::")));
+
+			for (int32 i = 0; i < ArgumentList.Num(); ++i)
+			{
+				UProperty* ArgProperty = ArgumentList[i];
+
+				if (i > 0)
+				{
+					Emit(Header, TEXT(", "));
+					Emit(Body, TEXT(", "));
+				}
+
+				if (ArgProperty->HasAnyPropertyFlags(CPF_OutParm))
+				{
+					Emit(Header, TEXT("/*out*/ "));
+					Emit(Body, TEXT("/*out*/ "));
+				}
+
+				ArgProperty->ExportCppDeclaration(Header, EExportedDeclaration::Parameter, NULL, EPropertyExportCPPFlags::CPPF_CustomTypeName | EPropertyExportCPPFlags::CPPF_BlueprintCppBackend);
+				ArgProperty->ExportCppDeclaration(Body, EExportedDeclaration::Parameter, NULL, EPropertyExportCPPFlags::CPPF_CustomTypeName | EPropertyExportCPPFlags::CPPF_BlueprintCppBackend);
+			}
+
+			Emit(Header, TEXT(")"));
+			if (bGenerateAsNativeEventImplementation)
+			{
+				Emit(Header, TEXT(" override"));
+			}
+			Emit(Header, TEXT(";\n"));
+			Emit(Body, TEXT(")\n"));
+		}
+
+		// Start the body of the implementation
+		Emit(Body, TEXT("{\n"));
+	}
+
+	if (!bGenerateStubOnly)
+	{
+		// Emit local variables
+		DeclareLocalVariables(FunctionContext, LocalVariables);
+
+		bool bUseSwitchState = false;
+		for (auto Node : FunctionContext.LinearExecutionList)
+		{
+			TArray<FBlueprintCompiledStatement*>* StatementList = FunctionContext.StatementsPerNode.Find(Node);
+			if (StatementList)
+			{
+				for (auto Statement : (*StatementList))
+				{
+					if (Statement && (
+						Statement->Type == KCST_UnconditionalGoto ||
+						Statement->Type == KCST_PushState ||
+						Statement->Type == KCST_GotoIfNot ||
+						Statement->Type == KCST_ComputedGoto ||
+						Statement->Type == KCST_EndOfThread ||
+						Statement->Type == KCST_EndOfThreadIfNot ||
+						Statement->Type == KCST_GotoReturn ||
+						Statement->Type == KCST_GotoReturnIfNot))
+					{
+						bUseSwitchState = true;
+						break;
+					}
+				}
+			}
+			if (bUseSwitchState)
+			{
+				break;
+			}
+		}
+
+		// Run thru code looking only at things marked as jump targets, to make sure the jump targets are ordered in order of appearance in the linear execution list
+		// Emit code in the order specified by the linear execution list (the first node is always the entry point for the function)
+		for (int32 NodeIndex = 0; NodeIndex < FunctionContext.LinearExecutionList.Num(); ++NodeIndex)
+		{
+			UEdGraphNode* StatementNode = FunctionContext.LinearExecutionList[NodeIndex];
+			TArray<FBlueprintCompiledStatement*>* StatementList = FunctionContext.StatementsPerNode.Find(StatementNode);
+
+			if (StatementList != NULL)
+			{
+				for (int32 StatementIndex = 0; StatementIndex < StatementList->Num(); ++StatementIndex)
+				{
+					FBlueprintCompiledStatement& Statement = *((*StatementList)[StatementIndex]);
+
+					if (Statement.bIsJumpTarget)
+					{
+						// Just making sure we number them in order of appearance, so jump statements don't influence the order
+						const int32 StateNum = StatementToStateIndex(FunctionContext, &Statement);
+					}
+				}
+			}
+		}
+
+		const FString FunctionImplementation = InnerFunctionImplementation(FunctionContext, bUseSwitchState);
+		Emit(Body, *FunctionImplementation);
+	}
+
+	const FString ReturnValueString = ReturnValue ? (FString(TEXT(" ")) + ReturnValue->GetName()) : TEXT("");
+	Emit(Body, *FString::Printf(TEXT("\treturn%s;\n"), *ReturnValueString));
+	Emit(Body, TEXT("}\n\n"));
+}
+
+void FBlueprintCompilerCppBackendBase::GenerateCodeFromEnum(UUserDefinedEnum* SourceEnum)
+{
+	check(SourceEnum);
+	const FString Name = SourceEnum->GetName();
+	EmitFileBeginning(Name, nullptr);
+
+	Emit(Header, TEXT("UENUM(BlueprintType)\nenum class "));
+	Emit(Header, *Name);
+	Emit(Header, TEXT(" : uint8\n{"));
+
+	for (int32 Index = 0; Index < SourceEnum->NumEnums(); ++Index)
+	{
+		const FString ElemName = SourceEnum->GetEnumName(Index);
+		const int32 ElemValue = Index;
+
+		const FString& DisplayNameMD = SourceEnum->GetMetaData(TEXT("DisplayName"), ElemValue);// TODO: value or index?
+		const FString Meta = DisplayNameMD.IsEmpty() ? FString() : FString::Printf(TEXT("UMETA(DisplayName = \"%s\")"), *DisplayNameMD);
+		Emit(Header, *FString::Printf(TEXT("\n\t%s = %d %s,"), *ElemName, ElemValue, *Meta));
+	}
+	Emit(Header, TEXT("\n};\n"));
+}
+
+void FBlueprintCompilerCppBackendBase::GenerateCodeFromStruct(UUserDefinedStruct* SourceStruct)
+{
+	check(SourceStruct);
+
+	EmitFileBeginning(SourceStruct->GetName(), SourceStruct);
+
+	const FString NewName = FString(TEXT("F")) + SourceStruct->GetName();
+	Emit(Header, TEXT("USTRUCT(BlueprintType)\n"));
+
+	Emit(Header, *FString::Printf(TEXT("struct %s\n{\npublic:\n\tGENERATED_BODY()\n"), *NewName));
+	EmitStructProperties(Header, SourceStruct);
+	Emit(Header, *FEmitDefaultValueHelper::GenerateGetDefaultValue(SourceStruct));
+	Emit(Header, TEXT("};\n"));
+}
+
+void FBlueprintCompilerCppBackendBase::EmitFileBeginning(const FString& CleanName, UStruct* SourceStruct)
+{
+	Emit(Header, TEXT("#pragma once\n\n"));
+	if (SourceStruct)
+	{ 
+		{
+			TArray<FString> PersistentHeaders;
+			PersistentHeaders.Add(FString(FApp::GetGameName()) + TEXT(".h"));
+			PersistentHeaders.Add(CleanName + TEXT(".h"));
+			PersistentHeaders.Add(TEXT("GeneratedCodeHelpers.h"));
+			Emit(Body, *FEmitHelper::GatherNativeHeadersToInclude(SourceStruct, PersistentHeaders));
+		}
+
+		// find objects referenced by functions/script
+		TArray<UObject*> IncludeInHeader;
+		TArray<UObject*> IncludeInBody;
+		{
+			FReferenceFinder HeaderReferenceFinder(IncludeInHeader, NULL, false, false, true, false);
+			FReferenceFinder BodyReferenceFinder(IncludeInBody, NULL, false, false, true, false);
+			{
+				TArray<UObject*> ObjectsToCheck;
+				GetObjectsWithOuter(SourceStruct, ObjectsToCheck, true);
+				for (auto Obj : ObjectsToCheck)
+				{
+					auto Property = Cast<UProperty>(Obj);
+					auto OwnerProperty = IsValid(Property) ? Property->GetOwnerProperty() : nullptr;
+					const bool bIsParam = OwnerProperty && (0 != (OwnerProperty->PropertyFlags & CPF_Parm));
+					const bool bIsMemberVariable = OwnerProperty && (OwnerProperty->GetOuter() == SourceStruct);
+					if (bIsParam || bIsMemberVariable)
+					{
+						HeaderReferenceFinder.FindReferences(Obj);
+					}
+					else
+					{
+						BodyReferenceFinder.FindReferences(Obj);
+					}
+				}
+			}
+
+			if (auto SuperStruct = SourceStruct->GetSuperStruct())
+			{
+				IncludeInHeader.AddUnique(SuperStruct);
+			}
+
+			if (auto SourceClass = Cast<UClass>(SourceStruct))
+			{
+				for (auto& ImplementedInterface : SourceClass->Interfaces)
+				{
+					IncludeInHeader.AddUnique(ImplementedInterface.Class);
+				}
+			}
+		}
+
+		TSet<FString> AlreadyIncluded;
+		AlreadyIncluded.Add(SourceStruct->GetName());
+
+		auto EmitInner = [&](FStringOutputDevice& Dst, const TArray<UObject*>& Src)
+		{
+			for (UObject* Obj : Src)
+			{
+				bool bWantedType = Obj && (Obj->IsA<UBlueprintGeneratedClass>() || Obj->IsA<UUserDefinedEnum>() || Obj->IsA<UUserDefinedStruct>());
+				if (!bWantedType)
+				{
+					for (UObject* OuterObj = (Obj ? Obj->GetOuter() : nullptr); nullptr != OuterObj; OuterObj = OuterObj->GetOuter())
+					{
+						if (OuterObj->IsA<UBlueprintGeneratedClass>() || OuterObj->IsA<UUserDefinedStruct>())
+						{
+							Obj = OuterObj;
+							bWantedType = true;
+							break;
+						}
+					}
+				}
+
+				if (!bWantedType)
+				{
+					Obj = Obj ? Cast<UBlueprintGeneratedClass>(Obj->GetClass()) : nullptr;
+					bWantedType = (nullptr != Obj);
+				}
+
+				if (bWantedType && Obj)
+				{
+					const FString Name = Obj->GetName();
+					if (!AlreadyIncluded.Contains(Name))
+					{
+						AlreadyIncluded.Add(Name);
+						Emit(Dst, *FString::Printf(TEXT("#include \"%s.h\"\n"), *Name));
+					}
+				}
+			}
+			Emit(Dst, TEXT("\n"));
+		};
+
+		EmitInner(Header, IncludeInHeader);
+		EmitInner(Body, IncludeInBody);
+	}
+	Emit(Header, *FString::Printf(TEXT("#include \"%s.generated.h\"\n\n"), *CleanName));
+}
diff --git a/Engine/Source/Developer/BlueprintCompilerCppBackend/Private/BlueprintCompilerCppBackendBase.h b/Engine/Source/Developer/BlueprintCompilerCppBackend/Private/BlueprintCompilerCppBackendBase.h
new file mode 100644
--- /dev/null
+++ b/Engine/Source/Developer/BlueprintCompilerCppBackend/Private/BlueprintCompilerCppBackendBase.h
@@ -0,0 +1,74 @@
+// Copyright 1998-2015 Epic Games, Inc. All Rights Reserved.
+#pragma  once
+
+#include "IBlueprintCompilerCppBackendModule.h"
+
+class FBlueprintCompilerCppBackendBase : public IBlueprintCompilerCppBackend
+{
+protected:
+	struct FFunctionLabelInfo
+	{
+		TMap<FBlueprintCompiledStatement*, int32> StateMap;
+		int32 StateCounter;
+
+		FFunctionLabelInfo()
+		{
+			StateCounter = 0;
+		}
+
+		int32 StatementToStateIndex(FBlueprintCompiledStatement* Statement)
+		{
+			int32& Index = StateMap.FindOrAdd(Statement);
+			if (Index == 0)
+			{
+				Index = ++StateCounter;
+			}
+
+			return Index;
+		}
+	};
+
+	TArray<FFunctionLabelInfo> StateMapPerFunction;
+	TMap<FKismetFunctionContext*, int32> FunctionIndexMap;
+
+	FString CppClassName;
+
+	FStringOutputDevice Header;
+	FStringOutputDevice Body;
+public:
+
+	// IBlueprintCompilerCppBackend implementation
+	virtual void GenerateCodeFromClass(UClass* SourceClass, TIndirectArray<FKismetFunctionContext>& Functions, bool bGenerateStubsOnly) override;
+	virtual void GenerateCodeFromEnum(UUserDefinedEnum* SourceEnum) override;
+	virtual void GenerateCodeFromStruct(UUserDefinedStruct* SourceStruct) override;
+
+	virtual const FString& GetBody()	const override { return Body; }
+	virtual const FString& GetHeader()	const override { return Header; }
+	// end of IBlueprintCompilerCppBackend implementation
+
+	virtual ~FBlueprintCompilerCppBackendBase() {}
+private:
+	void Emit(FStringOutputDevice& Target, const TCHAR* Message)
+	{
+		Target += Message;
+	}
+
+protected:
+	virtual FString InnerFunctionImplementation(FKismetFunctionContext& FunctionContext, bool bUseSwitchState) PURE_VIRTUAL(FBlueprintCompilerCppBackendBase::InnerFunctionImplementation, return FString(););
+
+	void EmitStructProperties(FStringOutputDevice& Target, UStruct* SourceClass);
+
+	/** Emits local variable declarations for a function */
+	void DeclareLocalVariables(FKismetFunctionContext& FunctionContext, TArray<UProperty*>& LocalVariables);
+	
+	/** Builds both the header declaration and body implementation of a function */
+	void ConstructFunction(FKismetFunctionContext& FunctionContext, bool bGenerateStubOnly);
+
+	void EmitFileBeginning(const FString& CleanName, UStruct* SourceStruct);
+
+	int32 StatementToStateIndex(FKismetFunctionContext& FunctionContext, FBlueprintCompiledStatement* Statement)
+	{
+		int32 Index = FunctionIndexMap.FindChecked(&FunctionContext);
+		return StateMapPerFunction[Index].StatementToStateIndex(Statement);
+	}
+};
diff --git a/Engine/Source/Developer/BlueprintCompilerCppBackend/Private/BlueprintCompilerCppBackendUtils.cpp b/Engine/Source/Developer/BlueprintCompilerCppBackend/Private/BlueprintCompilerCppBackendUtils.cpp
--- a/Engine/Source/Developer/BlueprintCompilerCppBackend/Private/BlueprintCompilerCppBackendUtils.cpp
+++ b/Engine/Source/Developer/BlueprintCompilerCppBackend/Private/BlueprintCompilerCppBackendUtils.cpp
@@ -2,78 +2,18 @@
 
 #include "BlueprintCompilerCppBackendModulePrivatePCH.h"
 #include "BlueprintCompilerCppBackendUtils.h"
-#include "BlueprintCompilerCppBackend.h"
 #include "EdGraphSchema_K2.h"
-#include "Editor/UnrealEd/Public/Kismet2/StructureEditorUtils.h"
-#include "Engine/InheritableComponentHandler.h"
-#include "Engine/DynamicBlueprintBinding.h"
 
-bool FSafeContextScopedEmmitter::IsSafeContextUsed() const
-{
-	return bSafeContextUsed;
-}
-
-FString FSafeContextScopedEmmitter::ValidationChain(FEmitterLocalContext& EmitterContext, const FBPTerminal* Term, FBlueprintCompilerCppBackend& CppBackend)
-{
-	TArray<FString> SafetyConditions;
-	for (; Term; Term = Term->Context)
-	{
-		if (!Term->IsStructContextType() && (Term->Type.PinSubCategory != TEXT("self")))
-		{
-			ensure(!Term->bIsLiteral);
-			SafetyConditions.Add(CppBackend.TermToText(EmitterContext, Term, false));
-		}
-	}
-
-	FString Result;
-	for (int32 Iter = SafetyConditions.Num() - 1; Iter >= 0; --Iter)
-	{
-		Result += FString(TEXT("IsValid("));
-		Result += SafetyConditions[Iter];
-		Result += FString(TEXT(")"));
-		if (Iter)
-		{
-			Result += FString(TEXT(" && "));
-		}
-	}
-
-	return Result;
-}
-
-FSafeContextScopedEmmitter::FSafeContextScopedEmmitter(FEmitterLocalContext& InEmitterContext, const FBPTerminal* Term, FBlueprintCompilerCppBackend& CppBackend)
-	: EmitterContext(InEmitterContext)
-	, bSafeContextUsed(false)
-{
-	const FString Conditions = ValidationChain(EmitterContext, Term, CppBackend);
-
-	if (!Conditions.IsEmpty())
-	{
-		bSafeContextUsed = true;
-		EmitterContext.AddLine(FString::Printf(TEXT("if(%s)"), *Conditions));
-		EmitterContext.AddLine(TEXT("{"));
-		EmitterContext.IncreaseIndent();
-	}
-}
-
-FSafeContextScopedEmmitter::~FSafeContextScopedEmmitter()
-{
-	if (bSafeContextUsed)
-	{
-		EmitterContext.DecreaseIndent();
-		EmitterContext.AddLine(TEXT("}"));
-	}
-}
-
-void FEmitHelper::ArrayToString(const TArray<FString>& Array, FString& String, const TCHAR* Separator)
+void FEmitHelper::ArrayToString(const TArray<FString>& Array, FString& OutString, const TCHAR* Separator)
 {
 	if (Array.Num())
 	{
-		String += Array[0];
+		OutString += Array[0];
 	}
 	for (int32 i = 1; i < Array.Num(); ++i)
 	{
-		String += TEXT(", ");
-		String += Array[i];
+		OutString += TEXT(", ");
+		OutString += Array[i];
 	}
 }
 
@@ -575,6 +515,11 @@ FString FEmitHelper::LiteralTerm(FEmitterLocalContext& EmitterContext, const FEd
 	{
 		if (LiteralObject)
 		{
+			const FString MappedObject = EmitterContext.FindGloballyMappedObject(LiteralObject);
+			if (!MappedObject.IsEmpty())
+			{
+				return MappedObject;
+			}
 			UClass* FoundClass = Cast<UClass>(Type.PinSubCategoryObject.Get());
 			FString ClassString = FoundClass ? (FString(FoundClass->GetPrefixCPP()) + FoundClass->GetName()) : TEXT("UObject");
 			return FString::Printf(TEXT("LoadObject<%s>(nullptr, TEXT(\"%s\"))"), *ClassString, *(LiteralObject->GetPathName().ReplaceCharWithEscapedChar()));
@@ -681,515 +626,3 @@ bool FEmitHelper::GenerateAssignmentCast(const FEdGraphPinType& LType, const FEd
 
 	return false;
 }
-
-void FEmitDefaultValueHelper::OuterGenerate(FEmitterLocalContext& Context
-	, const UProperty* Property
-	, const FString& OuterPath
-	, const uint8* DataContainer
-	, const uint8* OptionalDefaultDataContainer
-	, EPropertyAccessOperator AccessOperator
-	, bool bAllowProtected)
-{
-	if (Property->HasAnyPropertyFlags(CPF_EditorOnly | CPF_Transient))
-	{
-		UE_LOG(LogK2Compiler, Log, TEXT("FEmitDefaultValueHelper Skip EditorOnly or Transient property: %s"), *Property->GetPathName());
-		return;
-	}
-
-	for (int32 ArrayIndex = 0; ArrayIndex < Property->ArrayDim; ++ArrayIndex)
-	{
-		if (!OptionalDefaultDataContainer || !Property->Identical_InContainer(DataContainer, OptionalDefaultDataContainer, ArrayIndex))
-		{
-			FString PathToMember;
-			if (Property->HasAnyPropertyFlags(CPF_NativeAccessSpecifierPrivate) || (!bAllowProtected && Property->HasAnyPropertyFlags(CPF_NativeAccessSpecifierProtected)))
-			{
-				ensure(EPropertyAccessOperator::None != AccessOperator);
-				
-				const FString PropertyObject = Context.GenerateGetProperty(Property); //X::StaticClass()->FindPropertyByName(GET_MEMBER_NAME_CHECKED(X, X))
-
-				FStringOutputDevice TypeDeclaration;
-				const uint32 CppTemplateTypeFlags = EPropertyExportCPPFlags::CPPF_CustomTypeName
-					| EPropertyExportCPPFlags::CPPF_NoConst | EPropertyExportCPPFlags::CPPF_NoRef
-					| EPropertyExportCPPFlags::CPPF_BlueprintCppBackend;
-				Property->ExportCppDeclaration(TypeDeclaration, EExportedDeclaration::Parameter, nullptr, CppTemplateTypeFlags, true);
-
-				const FString OperatorStr = (EPropertyAccessOperator::Dot == AccessOperator) ? TEXT("&") : TEXT("");
-				const FString ContainerStr = (EPropertyAccessOperator::None == AccessOperator) ? TEXT("this") : FString::Printf(TEXT("%s(%s)"), *OperatorStr, *OuterPath);
-				PathToMember = Context.GenerateUniqueLocalName();
-				Context.AddLine(FString::Printf(TEXT("auto& %s = *(%s->ContainerPtrToValuePtr<%s>(%s, %d));")
-					, *PathToMember, *PropertyObject, *TypeDeclaration, *ContainerStr, ArrayIndex));
-			}
-			else
-			{
-				const FString AccessOperatorStr = (EPropertyAccessOperator::None == AccessOperator) ? TEXT("")
-					: ((EPropertyAccessOperator::Pointer == AccessOperator) ? TEXT("->") : TEXT("."));
-				const bool bStaticArray = (Property->ArrayDim > 1);
-				const FString ArrayPost = bStaticArray ? FString::Printf(TEXT("[%d]"), ArrayIndex) : TEXT("");
-				PathToMember = FString::Printf(TEXT("%s%s%s%s"), *OuterPath, *AccessOperatorStr, *Property->GetNameCPP(), *ArrayPost);
-			}
-			const uint8* ValuePtr = Property->ContainerPtrToValuePtr<uint8>(DataContainer, ArrayIndex);
-			const uint8* DefaultValuePtr = OptionalDefaultDataContainer ? Property->ContainerPtrToValuePtr<uint8>(OptionalDefaultDataContainer, ArrayIndex) : nullptr;
-			InnerGenerate(Context, Property, PathToMember, ValuePtr, DefaultValuePtr);
-		}
-	}
-}
-
-FString FEmitDefaultValueHelper::GenerateGetDefaultValue(const UUserDefinedStruct* Struct)
-{
-	check(Struct);
-	const FString StructName = FString(TEXT("F")) + Struct->GetName();
-	FString Result = FString::Printf(TEXT("\tstatic %s GetDefaultValue()\n\t{\n\t\t%s DefaultData__;\n"), *StructName, *StructName);
-
-	FStructOnScope StructData(Struct);
-	FStructureEditorUtils::Fill_MakeStructureDefaultValue(Struct, StructData.GetStructMemory());
-
-	FEmitterLocalContext Context;
-	Context.IncreaseIndent();
-	Context.IncreaseIndent();
-	for (auto Property : TFieldRange<const UProperty>(Struct))
-	{
-		OuterGenerate(Context, Property, TEXT("DefaultData__"), StructData.GetStructMemory(), nullptr, EPropertyAccessOperator::Dot);
-	}
-	Result += Context.GetResult();
-	Result += TEXT("\n\t\treturn DefaultData__;\n\t}\n");
-	return Result;
-}
-
-void FEmitDefaultValueHelper::InnerGenerate(FEmitterLocalContext& Context, const UProperty* Property, const FString& PathToMember, const uint8* ValuePtr, const uint8* DefaultValuePtr, bool bWithoutFirstConstructionLine)
-{
-	auto OneLineConstruction = [](FEmitterLocalContext& LocalContext, const UProperty* LocalProperty, const uint8* LocalValuePtr, FString& OutSingleLine, bool bGenerateEmptyStructConstructor) -> bool
-	{
-		bool bComplete = true;
-		FString ValueStr = HandleSpecialTypes(LocalContext, LocalProperty, LocalValuePtr);
-		if (ValueStr.IsEmpty())
-		{
-			auto StructProperty = Cast<const UStructProperty>(LocalProperty);
-			LocalProperty->ExportTextItem(ValueStr, LocalValuePtr, LocalValuePtr, nullptr, EPropertyPortFlags::PPF_ExportCpp);
-			if (ValueStr.IsEmpty() && StructProperty)
-			{
-				check(StructProperty->Struct);
-				if (bGenerateEmptyStructConstructor)
-				{
-					ValueStr = FString::Printf(TEXT("F%s{}"), *StructProperty->Struct->GetName()); //don;t override existing values
-				}
-				bComplete = false;
-			}
-			else if (ValueStr.IsEmpty())
-			{
-				UE_LOG(LogK2Compiler, Error, TEXT("FEmitDefaultValueHelper Cannot generate initilization: %s"), *LocalProperty->GetPathName());
-			}
-		}
-		OutSingleLine += ValueStr;
-		return bComplete;
-	};
-
-	auto StructProperty = Cast<const UStructProperty>(Property);
-	check(!StructProperty || StructProperty->Struct);
-	auto ArrayProperty = Cast<const UArrayProperty>(Property);
-	check(!ArrayProperty || ArrayProperty->Inner);
-
-	if (!bWithoutFirstConstructionLine)
-	{
-		FString ValueStr;
-		const bool bComplete = OneLineConstruction(Context, Property, ValuePtr, ValueStr, false);
-		if (!ValueStr.IsEmpty())
-		{
-			Context.AddLine(FString::Printf(TEXT("%s = %s;"), *PathToMember, *ValueStr));
-		}
-		// array initialization "array_var = TArray<..>()" is complete, but it still needs items.
-		if (bComplete && !ArrayProperty)
-		{
-			return;
-		}
-	}
-
-	if (StructProperty)
-	{
-		for (auto LocalProperty : TFieldRange<const UProperty>(StructProperty->Struct))
-		{
-			OuterGenerate(Context, LocalProperty, PathToMember, ValuePtr, DefaultValuePtr, EPropertyAccessOperator::Dot);
-		}
-	}
-	
-	if (ArrayProperty)
-	{
-		FScriptArrayHelper ScriptArrayHelper(ArrayProperty, ValuePtr);
-		for (int32 Index = 0; Index < ScriptArrayHelper.Num(); ++Index)
-		{
-			const uint8* LocalValuePtr = ScriptArrayHelper.GetRawPtr(Index);
-
-			FString ValueStr;
-			const bool bComplete = OneLineConstruction(Context, ArrayProperty->Inner, LocalValuePtr, ValueStr, true);
-			Context.AddLine(FString::Printf(TEXT("%s.Add(%s);"), *PathToMember, *ValueStr));
-			if (!bComplete)
-			{
-				const FString LocalPathToMember = FString::Printf(TEXT("%s[%d]"), *PathToMember, Index);
-				InnerGenerate(Context, ArrayProperty->Inner, LocalPathToMember, LocalValuePtr, nullptr, true);
-			}
-		}
-	}
-}
-
-FString FEmitDefaultValueHelper::HandleSpecialTypes(FEmitterLocalContext& Context, const UProperty* Property, const uint8* ValuePtr)
-{
-	//TODO: Use Path maps for Objects
-	if (auto ObjectProperty = Cast<UObjectProperty>(Property))
-	{
-		UObject* Object = ObjectProperty->GetPropertyValue(ValuePtr);
-		if (Object)
-		{
-			{
-				const FString MappedObject = Context.FindGlobalObject(Object);
-				if (!MappedObject.IsEmpty())
-				{
-					return MappedObject;
-				}
-			}
-
-			{
-				FString LocalMappedObject;
-				if (Context.bInsideConstructor && Context.FindLocalObject_InConstructor(Object, LocalMappedObject))
-				{
-					return LocalMappedObject;
-				}
-			}
-
-			{
-				auto BPGC = Context.GetCurrentlyGeneratedClass();
-				auto CDO = BPGC ? BPGC->GetDefaultObject(false) : nullptr;
-				if (BPGC && Object && CDO && Object->IsIn(BPGC) && !Object->IsIn(CDO) && Context.bCreatingObjectsPerClass)
-				{
-					return HandleClassSubobject(Context, Object);
-				}
-			}
-
-			if (!Context.bCreatingObjectsPerClass && Property->HasAnyPropertyFlags(CPF_InstancedReference))
-			{
-				const FString CreateAsInstancedSubobject = HandleInstancedSubobject(Context, Object);
-				if (!CreateAsInstancedSubobject.IsEmpty())
-				{
-					return CreateAsInstancedSubobject;
-				}
-			}
-		}
-		else if (ObjectProperty->HasMetaData(FBlueprintMetadata::MD_LatentCallbackTarget))
-		{
-			return TEXT("this");
-		}
-	}
-
-	if (auto StructProperty = Cast<UStructProperty>(Property))
-	{
-		if (TBaseStructure<FTransform>::Get() == StructProperty->Struct)
-		{
-			check(ValuePtr);
-			const FTransform* Transform = reinterpret_cast<const FTransform*>(ValuePtr);
-			const auto Rotation = Transform->GetRotation();
-			const auto Translation = Transform->GetTranslation();
-			const auto Scale = Transform->GetScale3D();
-			return FString::Printf(TEXT("FTransform(FQuat(%f, %f, %f, %f), FVector(%f, %f, %f), FVector(%f, %f, %f))")
-				, Rotation.X, Rotation.Y, Rotation.Z, Rotation.W
-				, Translation.X, Translation.Y, Translation.Z
-				, Scale.X, Scale.Y, Scale.Z);
-		}
-
-		if (TBaseStructure<FVector>::Get() == StructProperty->Struct)
-		{
-			const FVector* Vector = reinterpret_cast<const FVector*>(ValuePtr);
-			return FString::Printf(TEXT("FVector(%f, %f, %f)"), Vector->X, Vector->Y, Vector->Z);
-		}
-
-		if (TBaseStructure<FGuid>::Get() == StructProperty->Struct)
-		{
-			const FGuid* Guid = reinterpret_cast<const FGuid*>(ValuePtr);
-			return FString::Printf(TEXT("FGuid(0x%08X, 0x%08X, 0x%08X, 0x%08X)"), Guid->A, Guid->B, Guid->C, Guid->D);
-		}
-	}
-	return FString();
-}
-
-UActorComponent* FEmitDefaultValueHelper::HandleNonNativeComponent(FEmitterLocalContext& Context, UBlueprintGeneratedClass* BPGC, FName ObjectName, bool bNew, const FString& NativeName)
-{
-	USCS_Node* Node = nullptr;
-	for (UBlueprintGeneratedClass* NodeOwner = BPGC; NodeOwner && !Node; NodeOwner = Cast<UBlueprintGeneratedClass>(NodeOwner->GetSuperClass()))
-	{
-		Node = NodeOwner->SimpleConstructionScript ? NodeOwner->SimpleConstructionScript->FindSCSNode(ObjectName) : nullptr;
-	}
-
-	if (!Node)
-	{
-		return nullptr;
-	}
-
-	ensure(bNew == (Node->GetOuter() == BPGC->SimpleConstructionScript));
-	UActorComponent* ComponentTemplate = bNew
-		? Node->ComponentTemplate
-		: (BPGC->InheritableComponentHandler ? BPGC->InheritableComponentHandler->GetOverridenComponentTemplate(Node) : nullptr);
-
-	if (!ComponentTemplate)
-	{
-		return nullptr;
-	}
-
-	auto ComponentClass = ComponentTemplate->GetClass();
-	if (bNew)
-	{
-		Context.AddHighPriorityLine(FString::Printf(TEXT("%s = CreateDefaultSubobject<%s%s>(TEXT(\"%s\"));")
-			, *NativeName, ComponentClass->GetPrefixCPP(), *ComponentClass->GetName(), *ComponentTemplate->GetName()));
-		Context.AddObjectFromLocalProperty_InConstructor(ComponentTemplate, NativeName);
-	}
-
-	UObject* ParentTemplateComponent = bNew 
-		? nullptr 
-		: Node->GetActualComponentTemplate(Cast<UBlueprintGeneratedClass>(BPGC->GetSuperClass()));
-	ensure(bNew || ParentTemplateComponent);
-	ensure(ParentTemplateComponent != ComponentTemplate);
-	ensure(!ParentTemplateComponent || (ParentTemplateComponent->GetClass() == ComponentClass));
-
-	const UObject* ObjectToCompare = bNew
-		? ComponentClass->GetDefaultObject(false)
-		: ParentTemplateComponent;
-	for (auto Property : TFieldRange<const UProperty>(ComponentClass))
-	{
-		OuterGenerate(Context, Property, NativeName
-			, reinterpret_cast<const uint8*>(ComponentTemplate)
-			, reinterpret_cast<const uint8*>(ObjectToCompare)
-			, EPropertyAccessOperator::Pointer);
-	}
-
-	auto FindParentComponentName = [&]() -> FString
-	{
-		FString ResultName;
-		if (auto ParentNode = Node->GetSCS()->FindParentNode(Node))
-		{
-			const FName ParentVariableName = ParentNode->GetVariableName();
-			ResultName = (ParentVariableName != NAME_None) ? ParentVariableName.ToString() : FString();
-		}
-
-		if (ResultName.IsEmpty())
-		{
-			auto ParentTemplate = Node->GetParentComponentTemplate(CastChecked<UBlueprint>(BPGC->ClassGeneratedBy));
-			FString ParentComponentNativeName;
-			if (ParentTemplate && Context.FindLocalObject_InConstructor(ParentTemplateComponent, ParentComponentNativeName))
-			{
-				ResultName = ParentComponentNativeName;
-			}
-		}
-
-		if (ResultName.IsEmpty())
-		{
-			ResultName = (Node->ParentComponentOrVariableName != NAME_None) ? Node->ParentComponentOrVariableName.ToString() : FString();
-		}
-
-		return ResultName;
-	};
-
-	const FString ParentNodeName = FindParentComponentName();
-	if (!ParentNodeName.IsEmpty())
-	{
-		Context.AddLine(FString::Printf(TEXT("%s->AttachParent = %s;"), *NativeName, *ParentNodeName));
-	}
-
-	// TODO: Set root component
-	// TODO: Update local transform, etc..
-
-	return ComponentTemplate;
-}
-
-FString FEmitDefaultValueHelper::GenerateConstructor(UClass* InBPGC)
-{
-	auto BPGC = CastChecked<UBlueprintGeneratedClass>(InBPGC);
-	const FString CppClassName = FString(BPGC->GetPrefixCPP()) + BPGC->GetName();
-
-	FEmitterLocalContext Context;
-	Context.SetCurrentlyGeneratedClass(BPGC);
-	Context.AddLine(FString::Printf(TEXT("%s::%s(const FObjectInitializer& ObjectInitializer) : Super(ObjectInitializer)"), *CppClassName, *CppClassName));
-	Context.AddLine(TEXT("{"));
-	Context.IncreaseIndent();
-	Context.bInsideConstructor = true;
-	// When CDO is created create all subobjects owned by the class
-	TArray<UActorComponent*> ActorComponentTempatesOwnedByClass = BPGC->ComponentTemplates;
-	{
-		// Gather all CT from SCS and IH, the remaining ones are generated for class..
-		if (auto SCS = BPGC->SimpleConstructionScript)
-		{
-			for (auto Node : SCS->GetAllNodes())
-			{
-				ActorComponentTempatesOwnedByClass.RemoveSwap(Node->ComponentTemplate);
-			}
-		}
-		if (auto IH = BPGC->GetInheritableComponentHandler())
-		{
-			TArray<UActorComponent*> AllTemplates;
-			IH->GetAllTemplates(AllTemplates);
-			ActorComponentTempatesOwnedByClass.RemoveAllSwap([&](UActorComponent* Component) -> bool
-			{
-				return AllTemplates.Contains(Component);
-			});
-		}
-
-		Context.AddLine(FString::Printf(TEXT("if(HasAnyFlags(RF_ClassDefaultObject) && (%s::StaticClass() == GetClass()))"), *CppClassName));
-		Context.AddLine(TEXT("{"));
-		Context.IncreaseIndent();
-		Context.AddLine(TEXT("ensure(0 == GetClass()->ConvertedSubobjectsFromBPGC.Num());"));
-		Context.bCreatingObjectsPerClass = true;
-		Context.FlushLines();
-		for (auto ComponentTemplate : ActorComponentTempatesOwnedByClass)
-		{
-			if (ComponentTemplate)
-			{
-				HandleClassSubobject(Context, ComponentTemplate);
-			}
-		}
-
-		for (auto TimelineTemplate : BPGC->Timelines)
-		{
-			if (TimelineTemplate)
-			{
-				HandleClassSubobject(Context, TimelineTemplate);
-			}
-		}
-
-		for (auto DynamicBindingObject : BPGC->DynamicBindingObjects)
-		{
-			if (DynamicBindingObject)
-			{
-				HandleClassSubobject(Context, DynamicBindingObject);
-			}
-		}
-
-		Context.bCreatingObjectsPerClass = false;
-		Context.DecreaseIndent();
-		Context.AddLine(TEXT("}"));
-
-		// Let's have an easy access to generated class subobjects
-		Context.AddLine(TEXT("{")); // no shadow variables
-		Context.IncreaseIndent();
-		Context.FlushLines();
-		int32 SubobjectIndex = 0;
-		for (UObject* Obj : Context.GetObjectsCreatedPerClass_InConstructor())
-		{
-			FString NativeName;
-			const bool bFound = Context.FindLocalObject_InConstructor(Obj, NativeName);
-			ensure(bFound);
-			Context.AddHighPriorityLine(FString::Printf(TEXT("auto %s = CastChecked<%s%s>(GetClass()->ConvertedSubobjectsFromBPGC[%d]);")
-				, *NativeName
-				, Obj->GetClass()->GetPrefixCPP()
-				, *Obj->GetClass()->GetName()
-				, SubobjectIndex));
-			SubobjectIndex++;
-		}
-	}
-
-	UObject* CDO = BPGC->GetDefaultObject(false);
-	Context.AddObjectFromLocalProperty_InConstructor(CDO, TEXT("this"));
-
-	UObject* ParentCDO = BPGC->GetSuperClass()->GetDefaultObject(false);
-	check(CDO && ParentCDO);
-	Context.AddLine(TEXT(""));
-
-	// TIMELINES
-	TSet<const UProperty*> HandledProperties;
-	const bool bIsActor = BPGC->IsChildOf<AActor>();
-
-	// COMPONENTS
-	for (auto Property : TFieldRange<const UProperty>(BPGC))
-	{
-		const bool bNewProperty = Property->GetOwnerStruct() == BPGC;
-		const bool bIsAccessible = bNewProperty || !Property->HasAnyPropertyFlags(CPF_NativeAccessSpecifierPrivate);
-		if (bIsAccessible && !HandledProperties.Contains(Property))
-		{
-			auto ObjectProperty = Cast<UObjectProperty>(Property);
-			const bool bComponentProp = ObjectProperty && ObjectProperty->PropertyClass && ObjectProperty->PropertyClass->IsChildOf<UActorComponent>();
-			const bool bNullValue = ObjectProperty && (nullptr == ObjectProperty->GetPropertyValue_InContainer(CDO));
-			const bool bActorComponent = bIsActor && bComponentProp && bNullValue;
-			if (bIsActor && bComponentProp && bNullValue)
-			{
-				auto HandledComponent = HandleNonNativeComponent(Context, BPGC, Property->GetFName(), bNewProperty, Property->GetNameCPP());
-				if (HandledComponent)
-				{
-					ensure(!ActorComponentTempatesOwnedByClass.Contains(HandledComponent));
-					continue;
-				}
-			}
-
-			OuterGenerate(Context, Property, TEXT(""), reinterpret_cast<const uint8*>(CDO), bNewProperty ? nullptr : reinterpret_cast<const uint8*>(ParentCDO), EPropertyAccessOperator::None, true);
-		}
-	}
-
-	Context.DecreaseIndent();
-	Context.AddLine(TEXT("}"));
-	Context.DecreaseIndent();
-	Context.AddLine(TEXT("}"));
-	Context.bInsideConstructor = false;
-	return Context.GetResult();
-}
-
-FString FEmitDefaultValueHelper::HandleClassSubobject(FEmitterLocalContext& Context, UObject* Object)
-{
-	FString OuterStr;
-	if (!Context.FindLocalObject_InConstructor(Object->GetOuter(), OuterStr))
-	{
-		return FString();
-	}
-
-	const FString LocalNativeName = Context.AddNewObject_InConstructor(Object);
-	UClass* ObjectClass = Object->GetClass();
-	Context.AddHighPriorityLine(FString::Printf(
-		TEXT("auto %s = NewObject<%s%s>(%s, TEXT(\"%s\"));")
-		, *LocalNativeName
-		, ObjectClass->GetPrefixCPP()
-		, *ObjectClass->GetName()
-		, *OuterStr, *Object->GetName()));
-	if (Object->GetOuter() == Context.GetCurrentlyGeneratedClass())
-	{
-		Context.AddLine(FString::Printf(TEXT("GetClass()->ConvertedSubobjectsFromBPGC.Add(%s);"), *LocalNativeName));
-	}
-
-	for (auto Property : TFieldRange<const UProperty>(ObjectClass))
-	{
-		OuterGenerate(Context, Property, LocalNativeName
-			, reinterpret_cast<const uint8*>(Object)
-			, reinterpret_cast<const uint8*>(ObjectClass->GetDefaultObject(false))
-			, EPropertyAccessOperator::Pointer);
-	}
-
-	return LocalNativeName;
-}
-
-FString FEmitDefaultValueHelper::HandleInstancedSubobject(FEmitterLocalContext& Context, UObject* Object)
-{
-	check(Object);
-
-	const FString LocalNativeName = Context.AddNewObject_InConstructor(Object);
-	UClass* ObjectClass = Object->GetClass();
-
-	auto BPGC = Context.GetCurrentlyGeneratedClass();
-	auto CDO = BPGC ? BPGC->GetDefaultObject(false) : nullptr;
-	if (ensure(CDO) && (CDO == Object->GetOuter()))
-	{
-		Context.AddHighPriorityLine(FString::Printf(TEXT("auto %s = CreateDefaultSubobject<%s%s>(TEXT(\"%s\"));")
-			, *LocalNativeName, ObjectClass->GetPrefixCPP(), *ObjectClass->GetName(), *Object->GetName()));
-	}
-	else
-	{
-		FString OuterStr;
-		if (!Context.FindLocalObject_InConstructor(Object->GetOuter(), OuterStr))
-		{
-			return FString();
-		}
-		Context.AddHighPriorityLine(FString::Printf(TEXT("auto %s = NewObject<%s%s>(%s, TEXT(\"%s\"));")
-			, *LocalNativeName, ObjectClass->GetPrefixCPP(), *ObjectClass->GetName()
-			, *OuterStr, *Object->GetName()));
-	}
-
-	for (auto Property : TFieldRange<const UProperty>(ObjectClass))
-	{
-		OuterGenerate(Context, Property, LocalNativeName
-			, reinterpret_cast<const uint8*>(Object)
-			, reinterpret_cast<const uint8*>(ObjectClass->GetDefaultObject(false))
-			, EPropertyAccessOperator::Pointer);
-	}
-
-	return LocalNativeName;
-}
\ No newline at end of file
diff --git a/Engine/Source/Developer/BlueprintCompilerCppBackend/Private/BlueprintCompilerCppBackendUtils.h b/Engine/Source/Developer/BlueprintCompilerCppBackend/Private/BlueprintCompilerCppBackendUtils.h
--- a/Engine/Source/Developer/BlueprintCompilerCppBackend/Private/BlueprintCompilerCppBackendUtils.h
+++ b/Engine/Source/Developer/BlueprintCompilerCppBackend/Private/BlueprintCompilerCppBackendUtils.h
@@ -1,68 +1,5 @@
 // Copyright 1998-2015 Epic Games, Inc. All Rights Reserved.
 #pragma  once
-
-class FBlueprintCompilerCppBackend;
-
-// Generates single if scope. It's condition checks context of given term.
-struct FSafeContextScopedEmmitter
-{
-private:
-	FEmitterLocalContext& EmitterContext;
-	bool bSafeContextUsed;
-public:
-	bool IsSafeContextUsed() const;
-
-	static FString ValidationChain(FEmitterLocalContext& EmitterContext, const FBPTerminal* Term, FBlueprintCompilerCppBackend& CppBackend);
-
-	FSafeContextScopedEmmitter(FEmitterLocalContext& InEmitterContext, const FBPTerminal* Term, FBlueprintCompilerCppBackend& CppBackend);
-	~FSafeContextScopedEmmitter();
-};
-
-struct FEmitHelper
-{
-	static void ArrayToString(const TArray<FString>& Array, FString& String, const TCHAR* Separator);
-
-	static bool HasAllFlags(uint64 Flags, uint64 FlagsToCheck);
-
-	static FString HandleRepNotifyFunc(const UProperty* Property);
-
-	static bool MetaDataCanBeNative(const FName MetaDataName);
-
-	static FString HandleMetaData(const UField* Field, bool AddCategory = true, TArray<FString>* AdditinalMetaData = nullptr);
-
-	static TArray<FString> ProperyFlagsToTags(uint64 Flags);
-
-	static TArray<FString> FunctionFlagsToTags(uint64 Flags);
-
-	static bool IsBlueprintNativeEvent(uint64 FunctionFlags);
-
-	static bool IsBlueprintImplementableEvent(uint64 FunctionFlags);
-
-	static FString EmitUFuntion(UFunction* Function, TArray<FString>* AdditinalMetaData = nullptr);
-
-	static int32 ParseDelegateDetails(UFunction* Signature, FString& OutParametersMacro, FString& OutParamNumberStr);
-
-	static TArray<FString> EmitSinglecastDelegateDeclarations(const TArray<UDelegateProperty*>& Delegates);
-
-	static TArray<FString> EmitMulticastDelegateDeclarations(UClass* SourceClass);
-
-	static FString EmitLifetimeReplicatedPropsImpl(UClass* SourceClass, const FString& CppClassName, const TCHAR* InCurrentIndent);
-
-	static FString GatherNativeHeadersToInclude(UField* SourceItem, const TArray<FString>& PersistentHeaders);
-
-	static FString LiteralTerm(FEmitterLocalContext& EmitterContext, const FEdGraphPinType& Type, const FString& CustomValue, UObject* LiteralObject);
-
-	static FString DefaultValue(FEmitterLocalContext& EmitterContext, const FEdGraphPinType& Type);
-
-	static UFunction* GetOriginalFunction(UFunction* Function);
-
-	static bool ShoulsHandleAsNativeEvent(UFunction* Function);
-
-	static bool ShoulsHandleAsImplementableEvent(UFunction* Function);
-
-	static bool GenerateAssignmentCast(const FEdGraphPinType& LType, const FEdGraphPinType& RType, FString& OutCastBegin, FString& OutCastEnd);
-};
-
 struct FEmitterLocalContext
 {
 private:
@@ -165,7 +102,7 @@ public:
 		return ActualClass;
 	}
 
-	FString FindGlobalObject(UObject* Object)
+	FString FindGloballyMappedObject(UObject* Object)
 	{
 		// TODO: check if not excluded
 
@@ -257,6 +194,51 @@ public:
 	}
 };
 
+struct FEmitHelper
+{
+	static void ArrayToString(const TArray<FString>& Array, FString& OutString, const TCHAR* Separator);
+
+	static bool HasAllFlags(uint64 Flags, uint64 FlagsToCheck);
+
+	static FString HandleRepNotifyFunc(const UProperty* Property);
+
+	static bool MetaDataCanBeNative(const FName MetaDataName);
+
+	static FString HandleMetaData(const UField* Field, bool AddCategory = true, TArray<FString>* AdditinalMetaData = nullptr);
+
+	static TArray<FString> ProperyFlagsToTags(uint64 Flags);
+
+	static TArray<FString> FunctionFlagsToTags(uint64 Flags);
+
+	static bool IsBlueprintNativeEvent(uint64 FunctionFlags);
+
+	static bool IsBlueprintImplementableEvent(uint64 FunctionFlags);
+
+	static FString EmitUFuntion(UFunction* Function, TArray<FString>* AdditinalMetaData = nullptr);
+
+	static int32 ParseDelegateDetails(UFunction* Signature, FString& OutParametersMacro, FString& OutParamNumberStr);
+
+	static TArray<FString> EmitSinglecastDelegateDeclarations(const TArray<UDelegateProperty*>& Delegates);
+
+	static TArray<FString> EmitMulticastDelegateDeclarations(UClass* SourceClass);
+
+	static FString EmitLifetimeReplicatedPropsImpl(UClass* SourceClass, const FString& CppClassName, const TCHAR* InCurrentIndent);
+
+	static FString GatherNativeHeadersToInclude(UField* SourceItem, const TArray<FString>& PersistentHeaders);
+
+	static FString LiteralTerm(FEmitterLocalContext& EmitterContext, const FEdGraphPinType& Type, const FString& CustomValue, UObject* LiteralObject);
+
+	static FString DefaultValue(FEmitterLocalContext& EmitterContext, const FEdGraphPinType& Type);
+
+	static UFunction* GetOriginalFunction(UFunction* Function);
+
+	static bool ShoulsHandleAsNativeEvent(UFunction* Function);
+
+	static bool ShoulsHandleAsImplementableEvent(UFunction* Function);
+
+	static bool GenerateAssignmentCast(const FEdGraphPinType& LType, const FEdGraphPinType& RType, FString& OutCastBegin, FString& OutCastEnd);
+};
+
 struct FEmitDefaultValueHelper
 {
 	static FString GenerateGetDefaultValue(const UUserDefinedStruct* Struct);
diff --git a/Engine/Source/Developer/BlueprintCompilerCppBackend/Private/BlueprintCompilerCppBackendValueHelper.cpp b/Engine/Source/Developer/BlueprintCompilerCppBackend/Private/BlueprintCompilerCppBackendValueHelper.cpp
new file mode 100644
--- /dev/null
+++ b/Engine/Source/Developer/BlueprintCompilerCppBackend/Private/BlueprintCompilerCppBackendValueHelper.cpp
@@ -0,0 +1,519 @@
+// Copyright 1998-2015 Epic Games, Inc. All Rights Reserved.
+
+#include "BlueprintCompilerCppBackendModulePrivatePCH.h"
+#include "BlueprintCompilerCppBackendUtils.h"
+#include "Editor/UnrealEd/Public/Kismet2/StructureEditorUtils.h"
+#include "Engine/InheritableComponentHandler.h"
+#include "Engine/DynamicBlueprintBinding.h"
+
+void FEmitDefaultValueHelper::OuterGenerate(FEmitterLocalContext& Context
+	, const UProperty* Property
+	, const FString& OuterPath
+	, const uint8* DataContainer
+	, const uint8* OptionalDefaultDataContainer
+	, EPropertyAccessOperator AccessOperator
+	, bool bAllowProtected)
+{
+	if (Property->HasAnyPropertyFlags(CPF_EditorOnly | CPF_Transient))
+	{
+		UE_LOG(LogK2Compiler, Log, TEXT("FEmitDefaultValueHelper Skip EditorOnly or Transient property: %s"), *Property->GetPathName());
+		return;
+	}
+
+	for (int32 ArrayIndex = 0; ArrayIndex < Property->ArrayDim; ++ArrayIndex)
+	{
+		if (!OptionalDefaultDataContainer || !Property->Identical_InContainer(DataContainer, OptionalDefaultDataContainer, ArrayIndex))
+		{
+			FString PathToMember;
+			if (Property->HasAnyPropertyFlags(CPF_NativeAccessSpecifierPrivate) || (!bAllowProtected && Property->HasAnyPropertyFlags(CPF_NativeAccessSpecifierProtected)))
+			{
+				ensure(EPropertyAccessOperator::None != AccessOperator);
+				
+				const FString PropertyObject = Context.GenerateGetProperty(Property); //X::StaticClass()->FindPropertyByName(GET_MEMBER_NAME_CHECKED(X, X))
+
+				FStringOutputDevice TypeDeclaration;
+				const uint32 CppTemplateTypeFlags = EPropertyExportCPPFlags::CPPF_CustomTypeName
+					| EPropertyExportCPPFlags::CPPF_NoConst | EPropertyExportCPPFlags::CPPF_NoRef
+					| EPropertyExportCPPFlags::CPPF_BlueprintCppBackend;
+				Property->ExportCppDeclaration(TypeDeclaration, EExportedDeclaration::Parameter, nullptr, CppTemplateTypeFlags, true);
+
+				const FString OperatorStr = (EPropertyAccessOperator::Dot == AccessOperator) ? TEXT("&") : TEXT("");
+				const FString ContainerStr = (EPropertyAccessOperator::None == AccessOperator) ? TEXT("this") : FString::Printf(TEXT("%s(%s)"), *OperatorStr, *OuterPath);
+				PathToMember = Context.GenerateUniqueLocalName();
+				Context.AddLine(FString::Printf(TEXT("auto& %s = *(%s->ContainerPtrToValuePtr<%s>(%s, %d));")
+					, *PathToMember, *PropertyObject, *TypeDeclaration, *ContainerStr, ArrayIndex));
+			}
+			else
+			{
+				const FString AccessOperatorStr = (EPropertyAccessOperator::None == AccessOperator) ? TEXT("")
+					: ((EPropertyAccessOperator::Pointer == AccessOperator) ? TEXT("->") : TEXT("."));
+				const bool bStaticArray = (Property->ArrayDim > 1);
+				const FString ArrayPost = bStaticArray ? FString::Printf(TEXT("[%d]"), ArrayIndex) : TEXT("");
+				PathToMember = FString::Printf(TEXT("%s%s%s%s"), *OuterPath, *AccessOperatorStr, *Property->GetNameCPP(), *ArrayPost);
+			}
+			const uint8* ValuePtr = Property->ContainerPtrToValuePtr<uint8>(DataContainer, ArrayIndex);
+			const uint8* DefaultValuePtr = OptionalDefaultDataContainer ? Property->ContainerPtrToValuePtr<uint8>(OptionalDefaultDataContainer, ArrayIndex) : nullptr;
+			InnerGenerate(Context, Property, PathToMember, ValuePtr, DefaultValuePtr);
+		}
+	}
+}
+
+FString FEmitDefaultValueHelper::GenerateGetDefaultValue(const UUserDefinedStruct* Struct)
+{
+	check(Struct);
+	const FString StructName = FString(TEXT("F")) + Struct->GetName();
+	FString Result = FString::Printf(TEXT("\tstatic %s GetDefaultValue()\n\t{\n\t\t%s DefaultData__;\n"), *StructName, *StructName);
+
+	FStructOnScope StructData(Struct);
+	FStructureEditorUtils::Fill_MakeStructureDefaultValue(Struct, StructData.GetStructMemory());
+
+	FEmitterLocalContext Context;
+	Context.IncreaseIndent();
+	Context.IncreaseIndent();
+	for (auto Property : TFieldRange<const UProperty>(Struct))
+	{
+		OuterGenerate(Context, Property, TEXT("DefaultData__"), StructData.GetStructMemory(), nullptr, EPropertyAccessOperator::Dot);
+	}
+	Result += Context.GetResult();
+	Result += TEXT("\n\t\treturn DefaultData__;\n\t}\n");
+	return Result;
+}
+
+void FEmitDefaultValueHelper::InnerGenerate(FEmitterLocalContext& Context, const UProperty* Property, const FString& PathToMember, const uint8* ValuePtr, const uint8* DefaultValuePtr, bool bWithoutFirstConstructionLine)
+{
+	auto OneLineConstruction = [](FEmitterLocalContext& LocalContext, const UProperty* LocalProperty, const uint8* LocalValuePtr, FString& OutSingleLine, bool bGenerateEmptyStructConstructor) -> bool
+	{
+		bool bComplete = true;
+		FString ValueStr = HandleSpecialTypes(LocalContext, LocalProperty, LocalValuePtr);
+		if (ValueStr.IsEmpty())
+		{
+			auto StructProperty = Cast<const UStructProperty>(LocalProperty);
+			LocalProperty->ExportTextItem(ValueStr, LocalValuePtr, LocalValuePtr, nullptr, EPropertyPortFlags::PPF_ExportCpp);
+			if (ValueStr.IsEmpty() && StructProperty)
+			{
+				check(StructProperty->Struct);
+				if (bGenerateEmptyStructConstructor)
+				{
+					ValueStr = FString::Printf(TEXT("F%s{}"), *StructProperty->Struct->GetName()); //don;t override existing values
+				}
+				bComplete = false;
+			}
+			else if (ValueStr.IsEmpty())
+			{
+				UE_LOG(LogK2Compiler, Error, TEXT("FEmitDefaultValueHelper Cannot generate initilization: %s"), *LocalProperty->GetPathName());
+			}
+		}
+		OutSingleLine += ValueStr;
+		return bComplete;
+	};
+
+	auto StructProperty = Cast<const UStructProperty>(Property);
+	check(!StructProperty || StructProperty->Struct);
+	auto ArrayProperty = Cast<const UArrayProperty>(Property);
+	check(!ArrayProperty || ArrayProperty->Inner);
+
+	if (!bWithoutFirstConstructionLine)
+	{
+		FString ValueStr;
+		const bool bComplete = OneLineConstruction(Context, Property, ValuePtr, ValueStr, false);
+		if (!ValueStr.IsEmpty())
+		{
+			Context.AddLine(FString::Printf(TEXT("%s = %s;"), *PathToMember, *ValueStr));
+		}
+		// array initialization "array_var = TArray<..>()" is complete, but it still needs items.
+		if (bComplete && !ArrayProperty)
+		{
+			return;
+		}
+	}
+
+	if (StructProperty)
+	{
+		for (auto LocalProperty : TFieldRange<const UProperty>(StructProperty->Struct))
+		{
+			OuterGenerate(Context, LocalProperty, PathToMember, ValuePtr, DefaultValuePtr, EPropertyAccessOperator::Dot);
+		}
+	}
+	
+	if (ArrayProperty)
+	{
+		FScriptArrayHelper ScriptArrayHelper(ArrayProperty, ValuePtr);
+		for (int32 Index = 0; Index < ScriptArrayHelper.Num(); ++Index)
+		{
+			const uint8* LocalValuePtr = ScriptArrayHelper.GetRawPtr(Index);
+
+			FString ValueStr;
+			const bool bComplete = OneLineConstruction(Context, ArrayProperty->Inner, LocalValuePtr, ValueStr, true);
+			Context.AddLine(FString::Printf(TEXT("%s.Add(%s);"), *PathToMember, *ValueStr));
+			if (!bComplete)
+			{
+				const FString LocalPathToMember = FString::Printf(TEXT("%s[%d]"), *PathToMember, Index);
+				InnerGenerate(Context, ArrayProperty->Inner, LocalPathToMember, LocalValuePtr, nullptr, true);
+			}
+		}
+	}
+}
+
+FString FEmitDefaultValueHelper::HandleSpecialTypes(FEmitterLocalContext& Context, const UProperty* Property, const uint8* ValuePtr)
+{
+	//TODO: Use Path maps for Objects
+	if (auto ObjectProperty = Cast<UObjectProperty>(Property))
+	{
+		UObject* Object = ObjectProperty->GetPropertyValue(ValuePtr);
+		if (Object)
+		{
+			{
+				const FString MappedObject = Context.FindGloballyMappedObject(Object);
+				if (!MappedObject.IsEmpty())
+				{
+					return MappedObject;
+				}
+			}
+
+			{
+				FString LocalMappedObject;
+				if (Context.bInsideConstructor && Context.FindLocalObject_InConstructor(Object, LocalMappedObject))
+				{
+					return LocalMappedObject;
+				}
+			}
+
+			{
+				auto BPGC = Context.GetCurrentlyGeneratedClass();
+				auto CDO = BPGC ? BPGC->GetDefaultObject(false) : nullptr;
+				if (BPGC && Object && CDO && Object->IsIn(BPGC) && !Object->IsIn(CDO) && Context.bCreatingObjectsPerClass)
+				{
+					return HandleClassSubobject(Context, Object);
+				}
+			}
+
+			if (!Context.bCreatingObjectsPerClass && Property->HasAnyPropertyFlags(CPF_InstancedReference))
+			{
+				const FString CreateAsInstancedSubobject = HandleInstancedSubobject(Context, Object);
+				if (!CreateAsInstancedSubobject.IsEmpty())
+				{
+					return CreateAsInstancedSubobject;
+				}
+			}
+		}
+		else if (ObjectProperty->HasMetaData(FBlueprintMetadata::MD_LatentCallbackTarget))
+		{
+			return TEXT("this");
+		}
+	}
+
+	if (auto StructProperty = Cast<UStructProperty>(Property))
+	{
+		if (TBaseStructure<FTransform>::Get() == StructProperty->Struct)
+		{
+			check(ValuePtr);
+			const FTransform* Transform = reinterpret_cast<const FTransform*>(ValuePtr);
+			const auto Rotation = Transform->GetRotation();
+			const auto Translation = Transform->GetTranslation();
+			const auto Scale = Transform->GetScale3D();
+			return FString::Printf(TEXT("FTransform(FQuat(%f, %f, %f, %f), FVector(%f, %f, %f), FVector(%f, %f, %f))")
+				, Rotation.X, Rotation.Y, Rotation.Z, Rotation.W
+				, Translation.X, Translation.Y, Translation.Z
+				, Scale.X, Scale.Y, Scale.Z);
+		}
+
+		if (TBaseStructure<FVector>::Get() == StructProperty->Struct)
+		{
+			const FVector* Vector = reinterpret_cast<const FVector*>(ValuePtr);
+			return FString::Printf(TEXT("FVector(%f, %f, %f)"), Vector->X, Vector->Y, Vector->Z);
+		}
+
+		if (TBaseStructure<FGuid>::Get() == StructProperty->Struct)
+		{
+			const FGuid* Guid = reinterpret_cast<const FGuid*>(ValuePtr);
+			return FString::Printf(TEXT("FGuid(0x%08X, 0x%08X, 0x%08X, 0x%08X)"), Guid->A, Guid->B, Guid->C, Guid->D);
+		}
+	}
+	return FString();
+}
+
+UActorComponent* FEmitDefaultValueHelper::HandleNonNativeComponent(FEmitterLocalContext& Context, UBlueprintGeneratedClass* BPGC, FName ObjectName, bool bNew, const FString& NativeName)
+{
+	USCS_Node* Node = nullptr;
+	for (UBlueprintGeneratedClass* NodeOwner = BPGC; NodeOwner && !Node; NodeOwner = Cast<UBlueprintGeneratedClass>(NodeOwner->GetSuperClass()))
+	{
+		Node = NodeOwner->SimpleConstructionScript ? NodeOwner->SimpleConstructionScript->FindSCSNode(ObjectName) : nullptr;
+	}
+
+	if (!Node)
+	{
+		return nullptr;
+	}
+
+	ensure(bNew == (Node->GetOuter() == BPGC->SimpleConstructionScript));
+	UActorComponent* ComponentTemplate = bNew
+		? Node->ComponentTemplate
+		: (BPGC->InheritableComponentHandler ? BPGC->InheritableComponentHandler->GetOverridenComponentTemplate(Node) : nullptr);
+
+	if (!ComponentTemplate)
+	{
+		return nullptr;
+	}
+
+	auto ComponentClass = ComponentTemplate->GetClass();
+	if (bNew)
+	{
+		Context.AddHighPriorityLine(FString::Printf(TEXT("%s = CreateDefaultSubobject<%s%s>(TEXT(\"%s\"));")
+			, *NativeName, ComponentClass->GetPrefixCPP(), *ComponentClass->GetName(), *ComponentTemplate->GetName()));
+		Context.AddObjectFromLocalProperty_InConstructor(ComponentTemplate, NativeName);
+	}
+
+	UObject* ParentTemplateComponent = bNew 
+		? nullptr 
+		: Node->GetActualComponentTemplate(Cast<UBlueprintGeneratedClass>(BPGC->GetSuperClass()));
+	ensure(bNew || ParentTemplateComponent);
+	ensure(ParentTemplateComponent != ComponentTemplate);
+	ensure(!ParentTemplateComponent || (ParentTemplateComponent->GetClass() == ComponentClass));
+
+	const UObject* ObjectToCompare = bNew
+		? ComponentClass->GetDefaultObject(false)
+		: ParentTemplateComponent;
+	for (auto Property : TFieldRange<const UProperty>(ComponentClass))
+	{
+		OuterGenerate(Context, Property, NativeName
+			, reinterpret_cast<const uint8*>(ComponentTemplate)
+			, reinterpret_cast<const uint8*>(ObjectToCompare)
+			, EPropertyAccessOperator::Pointer);
+	}
+
+	auto FindParentComponentName = [&]() -> FString
+	{
+		FString ResultName;
+		if (auto ParentNode = Node->GetSCS()->FindParentNode(Node))
+		{
+			const FName ParentVariableName = ParentNode->GetVariableName();
+			ResultName = (ParentVariableName != NAME_None) ? ParentVariableName.ToString() : FString();
+		}
+
+		if (ResultName.IsEmpty())
+		{
+			auto ParentTemplate = Node->GetParentComponentTemplate(CastChecked<UBlueprint>(BPGC->ClassGeneratedBy));
+			FString ParentComponentNativeName;
+			if (ParentTemplate && Context.FindLocalObject_InConstructor(ParentTemplateComponent, ParentComponentNativeName))
+			{
+				ResultName = ParentComponentNativeName;
+			}
+		}
+
+		if (ResultName.IsEmpty())
+		{
+			ResultName = (Node->ParentComponentOrVariableName != NAME_None) ? Node->ParentComponentOrVariableName.ToString() : FString();
+		}
+
+		return ResultName;
+	};
+
+	const FString ParentNodeName = FindParentComponentName();
+	if (!ParentNodeName.IsEmpty())
+	{
+		Context.AddLine(FString::Printf(TEXT("%s->AttachParent = %s;"), *NativeName, *ParentNodeName));
+	}
+
+	// TODO: Set root component
+	// TODO: Update local transform, etc..
+
+	return ComponentTemplate;
+}
+
+FString FEmitDefaultValueHelper::GenerateConstructor(UClass* InBPGC)
+{
+	auto BPGC = CastChecked<UBlueprintGeneratedClass>(InBPGC);
+	const FString CppClassName = FString(BPGC->GetPrefixCPP()) + BPGC->GetName();
+
+	FEmitterLocalContext Context;
+	Context.SetCurrentlyGeneratedClass(BPGC);
+	Context.AddLine(FString::Printf(TEXT("%s::%s(const FObjectInitializer& ObjectInitializer) : Super(ObjectInitializer)"), *CppClassName, *CppClassName));
+	Context.AddLine(TEXT("{"));
+	Context.IncreaseIndent();
+	Context.bInsideConstructor = true;
+	// When CDO is created create all subobjects owned by the class
+	TArray<UActorComponent*> ActorComponentTempatesOwnedByClass = BPGC->ComponentTemplates;
+	{
+		// Gather all CT from SCS and IH, the remaining ones are generated for class..
+		if (auto SCS = BPGC->SimpleConstructionScript)
+		{
+			for (auto Node : SCS->GetAllNodes())
+			{
+				ActorComponentTempatesOwnedByClass.RemoveSwap(Node->ComponentTemplate);
+			}
+		}
+		if (auto IH = BPGC->GetInheritableComponentHandler())
+		{
+			TArray<UActorComponent*> AllTemplates;
+			IH->GetAllTemplates(AllTemplates);
+			ActorComponentTempatesOwnedByClass.RemoveAllSwap([&](UActorComponent* Component) -> bool
+			{
+				return AllTemplates.Contains(Component);
+			});
+		}
+
+		Context.AddLine(FString::Printf(TEXT("if(HasAnyFlags(RF_ClassDefaultObject) && (%s::StaticClass() == GetClass()))"), *CppClassName));
+		Context.AddLine(TEXT("{"));
+		Context.IncreaseIndent();
+		Context.AddLine(TEXT("ensure(0 == GetClass()->ConvertedSubobjectsFromBPGC.Num());"));
+		Context.bCreatingObjectsPerClass = true;
+		Context.FlushLines();
+		for (auto ComponentTemplate : ActorComponentTempatesOwnedByClass)
+		{
+			if (ComponentTemplate)
+			{
+				HandleClassSubobject(Context, ComponentTemplate);
+			}
+		}
+
+		for (auto TimelineTemplate : BPGC->Timelines)
+		{
+			if (TimelineTemplate)
+			{
+				HandleClassSubobject(Context, TimelineTemplate);
+			}
+		}
+
+		for (auto DynamicBindingObject : BPGC->DynamicBindingObjects)
+		{
+			if (DynamicBindingObject)
+			{
+				HandleClassSubobject(Context, DynamicBindingObject);
+			}
+		}
+
+		Context.bCreatingObjectsPerClass = false;
+		Context.DecreaseIndent();
+		Context.AddLine(TEXT("}"));
+
+		// Let's have an easy access to generated class subobjects
+		Context.AddLine(TEXT("{")); // no shadow variables
+		Context.IncreaseIndent();
+		Context.FlushLines();
+		int32 SubobjectIndex = 0;
+		for (UObject* Obj : Context.GetObjectsCreatedPerClass_InConstructor())
+		{
+			FString NativeName;
+			const bool bFound = Context.FindLocalObject_InConstructor(Obj, NativeName);
+			ensure(bFound);
+			Context.AddHighPriorityLine(FString::Printf(TEXT("auto %s = CastChecked<%s%s>(GetClass()->ConvertedSubobjectsFromBPGC[%d]);")
+				, *NativeName
+				, Obj->GetClass()->GetPrefixCPP()
+				, *Obj->GetClass()->GetName()
+				, SubobjectIndex));
+			SubobjectIndex++;
+		}
+	}
+
+	UObject* CDO = BPGC->GetDefaultObject(false);
+	Context.AddObjectFromLocalProperty_InConstructor(CDO, TEXT("this"));
+
+	UObject* ParentCDO = BPGC->GetSuperClass()->GetDefaultObject(false);
+	check(CDO && ParentCDO);
+	Context.AddLine(TEXT(""));
+
+	// TIMELINES
+	TSet<const UProperty*> HandledProperties;
+	const bool bIsActor = BPGC->IsChildOf<AActor>();
+
+	// COMPONENTS
+	for (auto Property : TFieldRange<const UProperty>(BPGC))
+	{
+		const bool bNewProperty = Property->GetOwnerStruct() == BPGC;
+		const bool bIsAccessible = bNewProperty || !Property->HasAnyPropertyFlags(CPF_NativeAccessSpecifierPrivate);
+		if (bIsAccessible && !HandledProperties.Contains(Property))
+		{
+			auto ObjectProperty = Cast<UObjectProperty>(Property);
+			const bool bComponentProp = ObjectProperty && ObjectProperty->PropertyClass && ObjectProperty->PropertyClass->IsChildOf<UActorComponent>();
+			const bool bNullValue = ObjectProperty && (nullptr == ObjectProperty->GetPropertyValue_InContainer(CDO));
+			const bool bActorComponent = bIsActor && bComponentProp && bNullValue;
+			if (bIsActor && bComponentProp && bNullValue)
+			{
+				auto HandledComponent = HandleNonNativeComponent(Context, BPGC, Property->GetFName(), bNewProperty, Property->GetNameCPP());
+				if (HandledComponent)
+				{
+					ensure(!ActorComponentTempatesOwnedByClass.Contains(HandledComponent));
+					continue;
+				}
+			}
+
+			OuterGenerate(Context, Property, TEXT(""), reinterpret_cast<const uint8*>(CDO), bNewProperty ? nullptr : reinterpret_cast<const uint8*>(ParentCDO), EPropertyAccessOperator::None, true);
+		}
+	}
+
+	Context.DecreaseIndent();
+	Context.AddLine(TEXT("}"));
+	Context.DecreaseIndent();
+	Context.AddLine(TEXT("}"));
+	Context.bInsideConstructor = false;
+	return Context.GetResult();
+}
+
+FString FEmitDefaultValueHelper::HandleClassSubobject(FEmitterLocalContext& Context, UObject* Object)
+{
+	FString OuterStr;
+	if (!Context.FindLocalObject_InConstructor(Object->GetOuter(), OuterStr))
+	{
+		return FString();
+	}
+
+	const FString LocalNativeName = Context.AddNewObject_InConstructor(Object);
+	UClass* ObjectClass = Object->GetClass();
+	Context.AddHighPriorityLine(FString::Printf(
+		TEXT("auto %s = NewObject<%s%s>(%s, TEXT(\"%s\"));")
+		, *LocalNativeName
+		, ObjectClass->GetPrefixCPP()
+		, *ObjectClass->GetName()
+		, *OuterStr, *Object->GetName()));
+	if (Object->GetOuter() == Context.GetCurrentlyGeneratedClass())
+	{
+		Context.AddLine(FString::Printf(TEXT("GetClass()->ConvertedSubobjectsFromBPGC.Add(%s);"), *LocalNativeName));
+	}
+
+	for (auto Property : TFieldRange<const UProperty>(ObjectClass))
+	{
+		OuterGenerate(Context, Property, LocalNativeName
+			, reinterpret_cast<const uint8*>(Object)
+			, reinterpret_cast<const uint8*>(ObjectClass->GetDefaultObject(false))
+			, EPropertyAccessOperator::Pointer);
+	}
+
+	return LocalNativeName;
+}
+
+FString FEmitDefaultValueHelper::HandleInstancedSubobject(FEmitterLocalContext& Context, UObject* Object)
+{
+	check(Object);
+
+	const FString LocalNativeName = Context.AddNewObject_InConstructor(Object);
+	UClass* ObjectClass = Object->GetClass();
+
+	auto BPGC = Context.GetCurrentlyGeneratedClass();
+	auto CDO = BPGC ? BPGC->GetDefaultObject(false) : nullptr;
+	if (ensure(CDO) && (CDO == Object->GetOuter()))
+	{
+		Context.AddHighPriorityLine(FString::Printf(TEXT("auto %s = CreateDefaultSubobject<%s%s>(TEXT(\"%s\"));")
+			, *LocalNativeName, ObjectClass->GetPrefixCPP(), *ObjectClass->GetName(), *Object->GetName()));
+	}
+	else
+	{
+		FString OuterStr;
+		if (!Context.FindLocalObject_InConstructor(Object->GetOuter(), OuterStr))
+		{
+			return FString();
+		}
+		Context.AddHighPriorityLine(FString::Printf(TEXT("auto %s = NewObject<%s%s>(%s, TEXT(\"%s\"));")
+			, *LocalNativeName, ObjectClass->GetPrefixCPP(), *ObjectClass->GetName()
+			, *OuterStr, *Object->GetName()));
+	}
+
+	for (auto Property : TFieldRange<const UProperty>(ObjectClass))
+	{
+		OuterGenerate(Context, Property, LocalNativeName
+			, reinterpret_cast<const uint8*>(Object)
+			, reinterpret_cast<const uint8*>(ObjectClass->GetDefaultObject(false))
+			, EPropertyAccessOperator::Pointer);
+	}
+
+	return LocalNativeName;
+}
\ No newline at end of file
diff --git a/Engine/Source/Editor/KismetCompiler/Private/KismetCompilerMisc.cpp b/Engine/Source/Editor/KismetCompiler/Private/KismetCompilerMisc.cpp
--- a/Engine/Source/Editor/KismetCompiler/Private/KismetCompilerMisc.cpp
+++ b/Engine/Source/Editor/KismetCompiler/Private/KismetCompilerMisc.cpp
@@ -1450,14 +1450,17 @@ void FKismetFunctionContext::MergeAdjacentStates()
 		}
 	}
 
-	// Remove unnecessary GotoReturn statements
-	// if it's last statement generated by last node (in LinearExecution) then it can be removed
-	const auto LastExecutedNode = LinearExecutionList.Num() ? LinearExecutionList.Last() : NULL;
-	TArray<FBlueprintCompiledStatement*>* StatementList = StatementsPerNode.Find(LastExecutedNode);
-	FBlueprintCompiledStatement* LastStatementInLastNode = (StatementList && StatementList->Num()) ? StatementList->Last() : NULL;
-	if (LastStatementInLastNode && (KCST_GotoReturn == LastStatementInLastNode->Type) && !LastStatementInLastNode->bIsJumpTarget)
-	{
-		StatementList->RemoveAt(StatementList->Num() - 1);
+	if (!bGeneratingCpp)
+	{
+		// Remove unnecessary GotoReturn statements
+		// if it's last statement generated by last node (in LinearExecution) then it can be removed
+		const auto LastExecutedNode = LinearExecutionList.Num() ? LinearExecutionList.Last() : NULL;
+		TArray<FBlueprintCompiledStatement*>* StatementList = StatementsPerNode.Find(LastExecutedNode);
+		FBlueprintCompiledStatement* LastStatementInLastNode = (StatementList && StatementList->Num()) ? StatementList->Last() : NULL;
+		if (LastStatementInLastNode && (KCST_GotoReturn == LastStatementInLastNode->Type) && !LastStatementInLastNode->bIsJumpTarget)
+		{
+			StatementList->RemoveAt(StatementList->Num() - 1);
+		}
 	}
 }
 
diff --git a/Engine/Source/Editor/KismetCompiler/Public/KismetCompiledFunctionContext.h b/Engine/Source/Editor/KismetCompiler/Public/KismetCompiledFunctionContext.h
--- a/Engine/Source/Editor/KismetCompiler/Public/KismetCompiledFunctionContext.h
+++ b/Engine/Source/Editor/KismetCompiler/Public/KismetCompiledFunctionContext.h
@@ -93,6 +93,7 @@ public:
 	// Stored calls of latent function (on current class), needed to tell if blueprint should be tickable
 	TArray< UK2Node_CallFunction* > LatentFunctionCalls;
 
+	//Skip some optimization. C++ code will be generated in this pass. 
 	bool bGeneratingCpp;
 
 	//Does this function use requires FlowStack ?
