UEBP-40 C++ code generation (WIP):
-comments added
- FEmitHelper::LiteralTerm uses FindGloballyMappedObject whenever it can

diff --git a/Engine/Source/Developer/BlueprintCompilerCppBackend/Private/BlueprintCompilerCppBackend.h b/Engine/Source/Developer/BlueprintCompilerCppBackend/Private/BlueprintCompilerCppBackend.h
--- a/Engine/Source/Developer/BlueprintCompilerCppBackend/Private/BlueprintCompilerCppBackend.h
+++ b/Engine/Source/Developer/BlueprintCompilerCppBackend/Private/BlueprintCompilerCppBackend.h
@@ -5,6 +5,7 @@
 
 struct FEmitterLocalContext;
 
+/** Default implementation of BlueprintCompilerCppBackend. Code generated for function body usually contains a big "switch". */
 class FBlueprintCompilerCppBackend : public FBlueprintCompilerCppBackendBase
 {
 protected:
diff --git a/Engine/Source/Developer/BlueprintCompilerCppBackend/Private/BlueprintCompilerCppBackendBase.h b/Engine/Source/Developer/BlueprintCompilerCppBackend/Private/BlueprintCompilerCppBackendBase.h
--- a/Engine/Source/Developer/BlueprintCompilerCppBackend/Private/BlueprintCompilerCppBackendBase.h
+++ b/Engine/Source/Developer/BlueprintCompilerCppBackend/Private/BlueprintCompilerCppBackendBase.h
@@ -3,6 +3,7 @@
 
 #include "IBlueprintCompilerCppBackendModule.h"
 
+/** The class generates all native code except of function body (notice InnerFunctionImplementation in not implemented) */
 class FBlueprintCompilerCppBackendBase : public IBlueprintCompilerCppBackend
 {
 protected:
@@ -54,6 +55,20 @@ private:
 	}
 
 protected:
+	/*
+		ConstructFunction generates definition, declaration and local variables. The actual implementation is generated by InnerFunctionImplementation.
+		For example - this code is generated by ConstructFunction:
+
+			void AMyActor_Child_C::BoxBox_Implementation(FBox2D NewParam,  FBox2D& NewParam1)
+			{
+				FBox2D CallFunc_BoxBox_NewParam1{};
+
+				// this part is generated by InnerFunctionImplementation
+
+				return;
+			}
+
+	*/
 	virtual FString InnerFunctionImplementation(FKismetFunctionContext& FunctionContext, bool bUseSwitchState) PURE_VIRTUAL(FBlueprintCompilerCppBackendBase::InnerFunctionImplementation, return FString(););
 
 	void EmitStructProperties(FStringOutputDevice& Target, UStruct* SourceClass);
diff --git a/Engine/Source/Developer/BlueprintCompilerCppBackend/Private/BlueprintCompilerCppBackendUtils.cpp b/Engine/Source/Developer/BlueprintCompilerCppBackend/Private/BlueprintCompilerCppBackendUtils.cpp
--- a/Engine/Source/Developer/BlueprintCompilerCppBackend/Private/BlueprintCompilerCppBackendUtils.cpp
+++ b/Engine/Source/Developer/BlueprintCompilerCppBackend/Private/BlueprintCompilerCppBackendUtils.cpp
@@ -501,15 +501,27 @@ FString FEmitHelper::LiteralTerm(FEmitterLocalContext& EmitterContext, const FEd
 	{
 		if (auto FoundClass = Cast<const UClass>(LiteralObject))
 		{
-			return FString::Printf(TEXT("%s%s::StaticClass()"), FoundClass->GetPrefixCPP(), *FoundClass->GetName());
+			const FString MappedObject = EmitterContext.FindGloballyMappedObject(LiteralObject);
+			if (!MappedObject.IsEmpty())
+			{
+				return MappedObject;
+			}
+			return FString::Printf(TEXT("LoadClass<UObject>(nullptr, TEXT(\"%s\"), nullptr, 0, nullptr "), *(LiteralObject->GetPathName().ReplaceCharWithEscapedChar()));
 		}
 		return FString(TEXT("nullptr"));
 	}
 	else if((UEdGraphSchema_K2::PC_AssetClass == Type.PinCategory) || (UEdGraphSchema_K2::PC_Asset == Type.PinCategory))
 	{
-		return LiteralObject
-			? FString::Printf(TEXT("FStringAssetReference(TEXT(\"%s\"))"), *(LiteralObject->GetPathName().ReplaceCharWithEscapedChar()))
-			: FString(TEXT("nullptr"));
+		if (LiteralObject)
+		{
+			const FString MappedObject = EmitterContext.FindGloballyMappedObject(LiteralObject);
+			if (!MappedObject.IsEmpty())
+			{
+				return MappedObject;
+			}
+			return FString::Printf(TEXT("FStringAssetReference(TEXT(\"%s\"))"), *(LiteralObject->GetPathName().ReplaceCharWithEscapedChar()));
+		}
+		return FString(TEXT("nullptr"));
 	}
 	else if (UEdGraphSchema_K2::PC_Object == Type.PinCategory)
 	{
