UEBP-40 C++ code generation (WIP):
- improved CustomThunk support
- MakeArray fix

diff --git a/Engine/Source/Developer/BlueprintCompilerCppBackend/Private/BlueprintCompilerCppBackend.cpp b/Engine/Source/Developer/BlueprintCompilerCppBackend/Private/BlueprintCompilerCppBackend.cpp
--- a/Engine/Source/Developer/BlueprintCompilerCppBackend/Private/BlueprintCompilerCppBackend.cpp
+++ b/Engine/Source/Developer/BlueprintCompilerCppBackend/Private/BlueprintCompilerCppBackend.cpp
@@ -171,7 +171,7 @@ void FBlueprintCompilerCppBackend::EmitObjectToBoolStatement(FEmitterLocalContex
 {
 	FString ObjectTarget = TermToText(EmitterContext, Statement.RHS[0]);
 	FString DestinationExpression = TermToText(EmitterContext, Statement.LHS);
-	EmitterContext.AddLine(FString::Printf(TEXT("%s = (%s != NULL);"), *DestinationExpression, *ObjectTarget));
+	EmitterContext.AddLine(FString::Printf(TEXT("%s = (nullptr != %s);"), *DestinationExpression, *ObjectTarget));
 }
 
 void FBlueprintCompilerCppBackend::EmitAddMulticastDelegateStatement(FEmitterLocalContext& EmitterContext, FKismetFunctionContext& FunctionContext, FBlueprintCompiledStatement& Statement)
@@ -224,8 +224,7 @@ void FBlueprintCompilerCppBackend::EmitCreateArrayStatement(FEmitterLocalContext
 	FBPTerminal* ArrayTerm = Statement.LHS;
 	const FString Array = TermToText(EmitterContext, ArrayTerm);
 
-	UArrayProperty* ArrayProperty = CastChecked<UArrayProperty>(ArrayTerm->AssociatedVarProperty);
-	UProperty* InnerProperty = ArrayProperty->Inner;
+	EmitterContext.AddLine(FString::Printf(TEXT("%s.SetNum(%d, true);"), *Array, Statement.RHS.Num()));
 
 	for (int32 i = 0; i < Statement.RHS.Num(); ++i)
 	{
@@ -315,11 +314,11 @@ FString FBlueprintCompilerCppBackend::EmitSwitchValueStatmentInner(FEmitterLocal
 
 	FStringOutputDevice IndexDeclaration;
 	check(IndexTerm && IndexTerm->AssociatedVarProperty);
-	IndexTerm->AssociatedVarProperty->ExportCppDeclaration(IndexDeclaration, EExportedDeclaration::Parameter, NULL, CppTemplateTypeFlags, true);
+	IndexTerm->AssociatedVarProperty->ExportCppDeclaration(IndexDeclaration, EExportedDeclaration::Parameter, nullptr, CppTemplateTypeFlags, true);
 
 	check(DefaultValueTerm && DefaultValueTerm->AssociatedVarProperty);
 	FStringOutputDevice ValueDeclaration;
-	DefaultValueTerm->AssociatedVarProperty->ExportCppDeclaration(ValueDeclaration, EExportedDeclaration::Parameter, NULL, CppTemplateTypeFlags, true);
+	DefaultValueTerm->AssociatedVarProperty->ExportCppDeclaration(ValueDeclaration, EExportedDeclaration::Parameter, nullptr, CppTemplateTypeFlags, true);
 
 	for (int32 TermIndex = TermsBeforeCases; TermIndex < (NumCases * TermsPerCase); TermIndex += TermsPerCase)
 	{
@@ -353,9 +352,9 @@ FString FBlueprintCompilerCppBackend::EmitMethodInputParameterList(FEmitterLocal
 			FString VarName;
 
 			FBPTerminal* Term = Statement.RHS[NumParams];
-			ensure(Term != NULL);
+			ensure(Term != nullptr);
 
-			if ((Statement.TargetLabel != NULL) && (Statement.UbergraphCallIndex == NumParams))
+			if ((Statement.TargetLabel != nullptr) && (Statement.UbergraphCallIndex == NumParams))
 			{
 				// The target label will only ever be set on a call function when calling into the Ubergraph or
 				// on a latent function that will later call into the ubergraph, either of which requires a patchup
@@ -435,7 +434,9 @@ FString FBlueprintCompilerCppBackend::EmitCallStatmentInner(FEmitterLocalContext
 	{
 		if (bStaticCall)
 		{
-			const bool bIsCustomThunk = Statement.FunctionToCall->HasMetaData(TEXT("CustomStructureParam")) || Statement.FunctionToCall->HasMetaData(TEXT("ArrayParm"));
+			const bool bIsCustomThunk = Statement.FunctionToCall->GetBoolMetaData(TEXT("CustomThunk")) 
+				|| Statement.FunctionToCall->HasMetaData(TEXT("CustomStructureParam")) 
+				|| Statement.FunctionToCall->HasMetaData(TEXT("ArrayParm"));
 			auto OwnerClass = Statement.FunctionToCall->GetOuterUClass();
 			Result += bIsCustomThunk ? TEXT("FCustomThunkTemplates::") : FString::Printf(TEXT("%s%s::"), OwnerClass->GetPrefixCPP(), *OwnerClass->GetName());
 		}
@@ -500,7 +501,7 @@ FString FBlueprintCompilerCppBackend::TermToText(FEmitterLocalContext& EmitterCo
 	else
 	{
 		FString ResultPath(TEXT(""));
-		if ((Term->Context != NULL) && (Term->Context->Name != PSC_Self))
+		if ((Term->Context != nullptr) && (Term->Context->Name != PSC_Self))
 		{
 			ResultPath = TermToText(EmitterContext, Term->Context, false);
 
diff --git a/Engine/Source/Developer/BlueprintCompilerCppBackend/Private/BlueprintCompilerCppBackendUtils.h b/Engine/Source/Developer/BlueprintCompilerCppBackend/Private/BlueprintCompilerCppBackendUtils.h
--- a/Engine/Source/Developer/BlueprintCompilerCppBackend/Private/BlueprintCompilerCppBackendUtils.h
+++ b/Engine/Source/Developer/BlueprintCompilerCppBackend/Private/BlueprintCompilerCppBackendUtils.h
@@ -102,6 +102,8 @@ public:
 		return ActualClass;
 	}
 
+	/** All objects (that can be referenced from other package) that will have a different path in cooked build 
+	(due to the native code generation), should be handled by this function */
 	FString FindGloballyMappedObject(UObject* Object)
 	{
 		// TODO: check if not excluded
diff --git a/Engine/Source/Programs/UnrealHeaderTool/Private/HeaderParser.cpp b/Engine/Source/Programs/UnrealHeaderTool/Private/HeaderParser.cpp
--- a/Engine/Source/Programs/UnrealHeaderTool/Private/HeaderParser.cpp
+++ b/Engine/Source/Programs/UnrealHeaderTool/Private/HeaderParser.cpp
@@ -5443,6 +5443,11 @@ void FHeaderParser::CompileFunctionDeclaration(FUnrealSourceFile& SourceFile, FC
 
 	ProcessFunctionSpecifiers(FuncInfo, SpecifiersFound);
 
+	if ((0 != (FuncInfo.FunctionExportFlags & FUNCEXPORT_CustomThunk)) && !MetaData.Contains("CustomThunk"))
+	{
+		MetaData.Add(TEXT("CustomThunk"), TEXT("true"));
+	}
+
 	if ((FuncInfo.FunctionFlags & FUNC_BlueprintPure) && GetCurrentClass()->HasAnyClassFlags(CLASS_Interface))
 	{
 		// Until pure interface casts are supported, we don't allow pures in interfaces
diff --git a/Engine/Source/Runtime/Engine/Public/GeneratedCodeHelpers.h b/Engine/Source/Runtime/Engine/Public/GeneratedCodeHelpers.h
--- a/Engine/Source/Runtime/Engine/Public/GeneratedCodeHelpers.h
+++ b/Engine/Source/Runtime/Engine/Public/GeneratedCodeHelpers.h
@@ -6,7 +6,8 @@
 #include "Kismet/KismetArrayLibrary.h"
 #include "Kismet/KismetSystemLibrary.h"
 #include "Kismet/KismetMathLibrary.h"
-
+#include "Kismet/BlueprintFunctionLibrary.h"
+#include "Kismet/KismetMathLibrary.h"
 // Special libraries
 #include "Kismet/DataTableFunctionLibrary.h"
 
@@ -194,12 +195,51 @@ public:
 	}
 
 	//Replacements for CustomThunk functions from UKismetSystemLibrary
-
 	template<typename T>
 	static void SetStructurePropertyByName(UObject* Object, FName PropertyName, const T& Value)
 	{
 		return UKismetSystemLibrary::Generic_SetStructurePropertyByName(Object, PropertyName, &Value);
 	}
+
+	static void SetCollisionProfileNameProperty(UObject* Object, FName PropertyName, const FCollisionProfileName& Value)
+	{
+		return UKismetSystemLibrary::Generic_SetStructurePropertyByName(Object, PropertyName, &Value);
+	}
+
+	// Replacements for CustomThunk functions from UBlueprintFunctionLibrary
+	static FStringAssetReference MakeStringAssetReference(const FString& AssetLongPathname)
+	{
+		FStringAssetReference Ref(AssetLongPathname);
+		if (!AssetLongPathname.IsEmpty() && !Ref.IsValid())
+		{
+			ExecutionMessage(*FString::Printf(TEXT("Asset path \"%s\" not valid. Only long path name is allowed."), *AssetLongPathname), ELogVerbosity::Error);
+			return FStringAssetReference();
+		}
+
+		return Ref;
+	}
+
+	// Replacements for CustomThunk functions from KismetMathLibrary
+	static float Divide_FloatFloat(float A, float B)
+	{
+		if (B == 0.f)
+		{
+			ExecutionMessage(TEXT("Divide by zero"), ELogVerbosity::Warning);
+			return 0.0f;
+		}
+		return UKismetMathLibrary::GenericDivide_FloatFloat(A, B);
+	}
+
+	static float Percent_FloatFloat(float A, float B)
+	{
+		if (B == 0.f)
+		{
+			ExecutionMessage(TEXT("Modulo by zero"), ELogVerbosity::Warning);
+			return 0.0f;
+		}
+
+		return UKismetMathLibrary::GenericPercent_FloatFloat(A, B);
+	}
 };
 
 template<typename IndexType, typename ValueType>
