typo fixed: Should in
ShoulsHandleAsImplementableEvent
ShoulsHandleAsNativeEvent

diff --git a/Engine/Source/Developer/BlueprintCompilerCppBackend/Private/BlueprintCompilerCppBackend.cpp b/Engine/Source/Developer/BlueprintCompilerCppBackend/Private/BlueprintCompilerCppBackend.cpp
--- a/Engine/Source/Developer/BlueprintCompilerCppBackend/Private/BlueprintCompilerCppBackend.cpp
+++ b/Engine/Source/Developer/BlueprintCompilerCppBackend/Private/BlueprintCompilerCppBackend.cpp
@@ -450,7 +450,7 @@ FString FBlueprintCompilerCppBackend::EmitCallStatmentInner(FEmitterLocalContext
 			Result += TEXT("Super::");
 		}
 		Result += Statement.FunctionToCall->GetName();
-		if (Statement.bIsParentContext && FEmitHelper::ShoulsHandleAsNativeEvent(Statement.FunctionToCall))
+		if (Statement.bIsParentContext && FEmitHelper::ShouldHandleAsNativeEvent(Statement.FunctionToCall))
 		{
 			ensure(!bCallOnDifferentObject);
 			Result += TEXT("_Implementation");
diff --git a/Engine/Source/Developer/BlueprintCompilerCppBackend/Private/BlueprintCompilerCppBackendBase.cpp b/Engine/Source/Developer/BlueprintCompilerCppBackend/Private/BlueprintCompilerCppBackendBase.cpp
--- a/Engine/Source/Developer/BlueprintCompilerCppBackend/Private/BlueprintCompilerCppBackendBase.cpp
+++ b/Engine/Source/Developer/BlueprintCompilerCppBackend/Private/BlueprintCompilerCppBackendBase.cpp
@@ -208,12 +208,12 @@ void FBlueprintCompilerCppBackendBase::ConstructFunction(FKismetFunctionContext&
 		bool bAddCppFromBpEventMD = false;
 		bool bGenerateAsNativeEventImplementation = false;
 		// Get original function declaration
-		if (FEmitHelper::ShoulsHandleAsNativeEvent(Function)) // BlueprintNativeEvent
+		if (FEmitHelper::ShouldHandleAsNativeEvent(Function)) // BlueprintNativeEvent
 		{
 			bGenerateAsNativeEventImplementation = true;
 			FunctionName += TEXT("_Implementation");
 		}
-		else if (FEmitHelper::ShoulsHandleAsImplementableEvent(Function)) // BlueprintImplementableEvent
+		else if (FEmitHelper::ShouldHandleAsImplementableEvent(Function)) // BlueprintImplementableEvent
 		{
 			bAddCppFromBpEventMD = true;
 		}
diff --git a/Engine/Source/Developer/BlueprintCompilerCppBackend/Private/BlueprintCompilerCppBackendUtils.cpp b/Engine/Source/Developer/BlueprintCompilerCppBackend/Private/BlueprintCompilerCppBackendUtils.cpp
--- a/Engine/Source/Developer/BlueprintCompilerCppBackend/Private/BlueprintCompilerCppBackendUtils.cpp
+++ b/Engine/Source/Developer/BlueprintCompilerCppBackend/Private/BlueprintCompilerCppBackendUtils.cpp
@@ -590,7 +590,7 @@ UFunction* FEmitHelper::GetOriginalFunction(UFunction* Function)
 	return Function;
 }
 
-bool FEmitHelper::ShoulsHandleAsNativeEvent(UFunction* Function)
+bool FEmitHelper::ShouldHandleAsNativeEvent(UFunction* Function)
 {
 	check(Function);
 	auto OriginalFunction = FEmitHelper::GetOriginalFunction(Function);
@@ -607,7 +607,7 @@ bool FEmitHelper::ShoulsHandleAsNativeEvent(UFunction* Function)
 	return false;
 }
 
-bool FEmitHelper::ShoulsHandleAsImplementableEvent(UFunction* Function)
+bool FEmitHelper::ShouldHandleAsImplementableEvent(UFunction* Function)
 {
 	check(Function);
 	auto OriginalFunction = FEmitHelper::GetOriginalFunction(Function);
diff --git a/Engine/Source/Developer/BlueprintCompilerCppBackend/Private/BlueprintCompilerCppBackendUtils.h b/Engine/Source/Developer/BlueprintCompilerCppBackend/Private/BlueprintCompilerCppBackendUtils.h
--- a/Engine/Source/Developer/BlueprintCompilerCppBackend/Private/BlueprintCompilerCppBackendUtils.h
+++ b/Engine/Source/Developer/BlueprintCompilerCppBackend/Private/BlueprintCompilerCppBackendUtils.h
@@ -234,9 +234,9 @@ struct FEmitHelper
 
 	static UFunction* GetOriginalFunction(UFunction* Function);
 
-	static bool ShoulsHandleAsNativeEvent(UFunction* Function);
+	static bool ShouldHandleAsNativeEvent(UFunction* Function);
 
-	static bool ShoulsHandleAsImplementableEvent(UFunction* Function);
+	static bool ShouldHandleAsImplementableEvent(UFunction* Function);
 
 	static bool GenerateAssignmentCast(const FEdGraphPinType& LType, const FEdGraphPinType& RType, FString& OutCastBegin, FString& OutCastEnd);
 };
