UEBP-40 C++ code generation (WIP): generated items have "ReplaceConverted" meta data with the original path (and redirected paths).

diff --git a/Engine/Source/Developer/BlueprintCompilerCppBackend/Private/BlueprintCompilerCppBackendBase.cpp b/Engine/Source/Developer/BlueprintCompilerCppBackend/Private/BlueprintCompilerCppBackendBase.cpp
--- a/Engine/Source/Developer/BlueprintCompilerCppBackend/Private/BlueprintCompilerCppBackendBase.cpp
+++ b/Engine/Source/Developer/BlueprintCompilerCppBackend/Private/BlueprintCompilerCppBackendBase.cpp
@@ -79,7 +79,9 @@ void FBlueprintCompilerCppBackendBase::GenerateCodeFromClass(UClass* SourceClass
 	const bool bIsInterface = SourceClass->IsChildOf<UInterface>();
 	if (bIsInterface)
 	{
-		Emit(Header, *FString::Printf(TEXT("UINTERFACE(Blueprintable)\nclass U%s : public UInterface\n{\n\tGENERATED_BODY()\n};\n"), *CleanCppClassName));
+		Emit(Header, TEXT("UINTERFACE(Blueprintable"));
+		EmitReplaceConvertedMetaData(SourceClass);
+		Emit(Header, *FString::Printf(TEXT(")\nclass U%s : public UInterface\n{\n\tGENERATED_BODY()\n};\n"), *CleanCppClassName));
 		Emit(Header, *FString::Printf(TEXT("\nclass I%s"), *CleanCppClassName));
 	}
 	else
@@ -89,6 +91,7 @@ void FBlueprintCompilerCppBackendBase::GenerateCodeFromClass(UClass* SourceClass
 		{
 			Emit(Header, TEXT("Blueprintable, BlueprintType"));
 		}
+		EmitReplaceConvertedMetaData(SourceClass);
 		Emit(Header, TEXT(")\n"));
 
 		UClass* SuperClass = SourceClass->GetSuperClass();
@@ -327,7 +330,9 @@ void FBlueprintCompilerCppBackendBase::GenerateCodeFromEnum(UUserDefinedEnum* So
 	const FString Name = SourceEnum->GetName();
 	EmitFileBeginning(Name, nullptr);
 
-	Emit(Header, TEXT("UENUM(BlueprintType)\nenum class "));
+	Emit(Header, TEXT("UENUM(BlueprintType"));
+	EmitReplaceConvertedMetaData(SourceEnum);
+	Emit(Header, TEXT(")\nenum class "));
 	Emit(Header, *Name);
 	Emit(Header, TEXT(" : uint8\n{"));
 
@@ -343,6 +348,18 @@ void FBlueprintCompilerCppBackendBase::GenerateCodeFromEnum(UUserDefinedEnum* So
 	Emit(Header, TEXT("\n};\n"));
 }
 
+void FBlueprintCompilerCppBackendBase::EmitReplaceConvertedMetaData(UObject* Obj)
+{
+	const FString ReplaceConvertedMD = FEmitHelper::GenerateReplaceConvertedMD(Obj);
+	if (!ReplaceConvertedMD.IsEmpty())
+	{
+		TArray<FString> AdditionalMD;
+		AdditionalMD.Add(ReplaceConvertedMD);
+		Emit(Header, TEXT(","));
+		Emit(Header, *FEmitHelper::HandleMetaData(nullptr, false, &AdditionalMD));
+	}
+}
+
 void FBlueprintCompilerCppBackendBase::GenerateCodeFromStruct(UUserDefinedStruct* SourceStruct)
 {
 	check(SourceStruct);
@@ -350,7 +367,9 @@ void FBlueprintCompilerCppBackendBase::GenerateCodeFromStruct(UUserDefinedStruct
 	EmitFileBeginning(SourceStruct->GetName(), SourceStruct);
 
 	const FString NewName = FString(TEXT("F")) + SourceStruct->GetName();
-	Emit(Header, TEXT("USTRUCT(BlueprintType)\n"));
+	Emit(Header, TEXT("USTRUCT(BlueprintType"));
+	EmitReplaceConvertedMetaData(SourceStruct);
+	Emit(Header, TEXT(")\n"));
 
 	Emit(Header, *FString::Printf(TEXT("struct %s\n{\npublic:\n\tGENERATED_BODY()\n"), *NewName));
 	EmitStructProperties(Header, SourceStruct);
diff --git a/Engine/Source/Developer/BlueprintCompilerCppBackend/Private/BlueprintCompilerCppBackendBase.h b/Engine/Source/Developer/BlueprintCompilerCppBackend/Private/BlueprintCompilerCppBackendBase.h
--- a/Engine/Source/Developer/BlueprintCompilerCppBackend/Private/BlueprintCompilerCppBackendBase.h
+++ b/Engine/Source/Developer/BlueprintCompilerCppBackend/Private/BlueprintCompilerCppBackendBase.h
@@ -86,4 +86,6 @@ protected:
 		int32 Index = FunctionIndexMap.FindChecked(&FunctionContext);
 		return StateMapPerFunction[Index].StatementToStateIndex(Statement);
 	}
+
+	void EmitReplaceConvertedMetaData(UObject* Obj);
 };
diff --git a/Engine/Source/Developer/BlueprintCompilerCppBackend/Private/BlueprintCompilerCppBackendUtils.cpp b/Engine/Source/Developer/BlueprintCompilerCppBackend/Private/BlueprintCompilerCppBackendUtils.cpp
--- a/Engine/Source/Developer/BlueprintCompilerCppBackend/Private/BlueprintCompilerCppBackendUtils.cpp
+++ b/Engine/Source/Developer/BlueprintCompilerCppBackend/Private/BlueprintCompilerCppBackendUtils.cpp
@@ -49,12 +49,9 @@ FString FEmitHelper::HandleMetaData(const UField* Field, bool AddCategory, TArra
 {
 	FString MetaDataStr;
 
-	check(Field);
-	UPackage* Package = Field->GetOutermost();
-	check(Package);
-	const UMetaData* MetaData = Package->GetMetaData();
-	check(MetaData);
-	const TMap<FName, FString>* ValuesMap = MetaData->ObjectMetaDataMap.Find(Field);
+	UPackage* Package = Field ? Field->GetOutermost() : nullptr;
+	const UMetaData* MetaData = Package ? Package->GetMetaData() : nullptr;
+	const TMap<FName, FString>* ValuesMap = MetaData ? MetaData->ObjectMetaDataMap.Find(Field) : nullptr;
 	TArray<FString> MetaDataStrings;
 	if (ValuesMap && ValuesMap->Num())
 	{
@@ -210,6 +207,32 @@ bool FEmitHelper::IsBlueprintImplementableEvent(uint64 FunctionFlags)
 	return HasAllFlags(FunctionFlags, FUNC_Event | FUNC_BlueprintEvent) && !HasAllFlags(FunctionFlags, FUNC_Native);
 }
 
+FString FEmitHelper::GenerateReplaceConvertedMD(UObject* Obj)
+{
+	FString Result;
+	if (Obj)
+	{
+		Result = TEXT("ReplaceConverted=\"");
+		Result += Obj->GetPathName();
+		UPackage* Package = Obj->GetOutermost();
+		check(Package);
+		TArray<UObject*> AllObjects;
+		GetObjectsWithOuter(Package, AllObjects, true);
+		for (auto LocalObj : AllObjects)
+		{
+			auto Redirector = Cast<UObjectRedirector>(LocalObj);
+			if (Redirector && (Obj == Redirector->DestinationObject))
+			{
+				Result += TEXT(",");
+				Result += Redirector->GetPathName();
+			}
+		}
+		Result += TEXT("\"");
+	}
+
+	return Result;
+}
+
 FString FEmitHelper::EmitUFuntion(UFunction* Function, TArray<FString>* AdditinalMetaData)
 {
 	TArray<FString> Tags = FEmitHelper::FunctionFlagsToTags(Function->FunctionFlags);
diff --git a/Engine/Source/Developer/BlueprintCompilerCppBackend/Private/BlueprintCompilerCppBackendUtils.h b/Engine/Source/Developer/BlueprintCompilerCppBackend/Private/BlueprintCompilerCppBackendUtils.h
--- a/Engine/Source/Developer/BlueprintCompilerCppBackend/Private/BlueprintCompilerCppBackendUtils.h
+++ b/Engine/Source/Developer/BlueprintCompilerCppBackend/Private/BlueprintCompilerCppBackendUtils.h
@@ -239,6 +239,8 @@ struct FEmitHelper
 	static bool ShouldHandleAsImplementableEvent(UFunction* Function);
 
 	static bool GenerateAssignmentCast(const FEdGraphPinType& LType, const FEdGraphPinType& RType, FString& OutCastBegin, FString& OutCastEnd);
+
+	static FString FEmitHelper::GenerateReplaceConvertedMD(UObject* Obj);
 };
 
 struct FEmitDefaultValueHelper
