UEBP-40 C++ code generation (WIP):
- fixed #include list
- improved automatic casts
- temporary enum fix

diff --git a/Engine/Source/Developer/BlueprintCompilerCppBackend/Private/BlueprintCompilerCppBackend.cpp b/Engine/Source/Developer/BlueprintCompilerCppBackend/Private/BlueprintCompilerCppBackend.cpp
--- a/Engine/Source/Developer/BlueprintCompilerCppBackend/Private/BlueprintCompilerCppBackend.cpp
+++ b/Engine/Source/Developer/BlueprintCompilerCppBackend/Private/BlueprintCompilerCppBackend.cpp
@@ -100,7 +100,7 @@ void FBlueprintCompilerCppBackend::EmitAssignmentStatment(FEmitterLocalContext&
 
 	FString BeginCast;
 	FString EndCast;
-	FEmitHelper::GenerateAssignmentCast(Statement.LHS->Type, Statement.RHS[0]->Type, BeginCast, EndCast);
+	FEmitHelper::GenerateAutomaticCast(Statement.LHS->Type, Statement.RHS[0]->Type, BeginCast, EndCast);
 	EmitterContext.AddLine(FString::Printf(TEXT("%s = %s%s%s;"), *DestinationExpression, *BeginCast, *SourceExpression, *EndCast));
 }
 
@@ -352,7 +352,7 @@ FString FBlueprintCompilerCppBackend::EmitMethodInputParameterList(FEmitterLocal
 			FString VarName;
 
 			FBPTerminal* Term = Statement.RHS[NumParams];
-			ensure(Term != nullptr);
+			check(Term != nullptr);
 
 			if ((Statement.TargetLabel != nullptr) && (Statement.UbergraphCallIndex == NumParams))
 			{
@@ -373,7 +373,18 @@ FString FBlueprintCompilerCppBackend::EmitMethodInputParameterList(FEmitterLocal
 			else
 			{
 				// Emit a normal parameter term
-				VarName = TermToText(EmitterContext, Term);
+				FString BeginCast;
+				FString CloseCast;
+				FEdGraphPinType LType;
+				auto Schema = GetDefault<UEdGraphSchema_K2>();
+				check(Schema);
+				if (Schema->ConvertPropertyToPinType(FuncParamProperty, LType))
+				{
+					FEmitHelper::GenerateAutomaticCast(LType, Term->Type, BeginCast, CloseCast);
+				}
+				VarName += BeginCast;
+				VarName += TermToText(EmitterContext, Term);
+				VarName += CloseCast;
 			}
 
 			if (FuncParamProperty->HasAnyPropertyFlags(CPF_OutParm))
@@ -397,7 +408,7 @@ FString FBlueprintCompilerCppBackend::EmitCallStatmentInner(FEmitterLocalContext
 	const bool bInterfaceCall = bCallOnDifferentObject && Statement.FunctionContext && (UEdGraphSchema_K2::PC_Interface == Statement.FunctionContext->Type.PinCategory);
 
 	FString Result;
-	bool bCloseCast = false;
+	FString CloseCast;
 	if (!bInline)
 	{
 		// Handle the return value of the function being called
@@ -405,18 +416,14 @@ FString FBlueprintCompilerCppBackend::EmitCallStatmentInner(FEmitterLocalContext
 		if (FuncToCallReturnProperty && ensure(Statement.LHS))
 		{
 			FString BeginCast;
-			if (auto ObjectProperty = Cast<UObjectProperty>(FuncToCallReturnProperty))
+			FEdGraphPinType RType;
+			auto Schema = GetDefault<UEdGraphSchema_K2>();
+			check(Schema);
+			if (Schema->ConvertPropertyToPinType(FuncToCallReturnProperty, RType))
 			{
-				UClass* LClass = Statement.LHS ? Cast<UClass>(Statement.LHS->Type.PinSubCategoryObject.Get()) : nullptr;
-				if (LClass && LClass->IsChildOf(ObjectProperty->PropertyClass) && !ObjectProperty->PropertyClass->IsChildOf(LClass))
-				{
-					BeginCast = FString::Printf(TEXT("CastChecked<%s%s>("), LClass->GetPrefixCPP(), *LClass->GetName());
-					bCloseCast = true;
-				}
+				FEmitHelper::GenerateAutomaticCast(Statement.LHS->Type, RType, BeginCast, CloseCast);
 			}
-
 			Result += FString::Printf(TEXT("%s = %s"), *TermToText(EmitterContext, Statement.LHS), *BeginCast);
-
 		}
 	}
 
@@ -461,12 +468,7 @@ FString FBlueprintCompilerCppBackend::EmitCallStatmentInner(FEmitterLocalContext
 	}
 	Result += EmitMethodInputParameterList(EmitterContext, Statement);
 	Result += TEXT(")");
-
-	if (bCloseCast)
-	{
-		Result += TEXT(", ECastCheckedType::NullAllowed)");
-	}
-
+	Result += CloseCast;
 	if (!bInline)
 	{
 		Result += TEXT(";");
diff --git a/Engine/Source/Developer/BlueprintCompilerCppBackend/Private/BlueprintCompilerCppBackendBase.cpp b/Engine/Source/Developer/BlueprintCompilerCppBackend/Private/BlueprintCompilerCppBackendBase.cpp
--- a/Engine/Source/Developer/BlueprintCompilerCppBackend/Private/BlueprintCompilerCppBackendBase.cpp
+++ b/Engine/Source/Developer/BlueprintCompilerCppBackend/Private/BlueprintCompilerCppBackendBase.cpp
@@ -382,14 +382,6 @@ void FBlueprintCompilerCppBackendBase::EmitFileBeginning(const FString& CleanNam
 	Emit(Header, TEXT("#pragma once\n\n"));
 	if (SourceStruct)
 	{ 
-		{
-			TArray<FString> PersistentHeaders;
-			PersistentHeaders.Add(FString(FApp::GetGameName()) + TEXT(".h"));
-			PersistentHeaders.Add(CleanName + TEXT(".h"));
-			PersistentHeaders.Add(TEXT("GeneratedCodeHelpers.h"));
-			Emit(Body, *FEmitHelper::GatherNativeHeadersToInclude(SourceStruct, PersistentHeaders));
-		}
-
 		// find objects referenced by functions/script
 		TArray<UObject*> IncludeInHeader;
 		TArray<UObject*> IncludeInBody;
@@ -430,11 +422,22 @@ void FBlueprintCompilerCppBackendBase::EmitFileBeginning(const FString& CleanNam
 			}
 		}
 
+		auto EmitIncludeHeader = [&](FStringOutputDevice& Dst, const TCHAR* Message, bool bAddDotH)
+		{
+			Emit(Dst, *FString::Printf(TEXT("#include \"%s%s\"\n"), Message, bAddDotH ? TEXT(".h") : TEXT("")));
+		};
+		EmitIncludeHeader(Body, FApp::GetGameName(), true);
+		EmitIncludeHeader(Body, *CleanName, true);
+		EmitIncludeHeader(Body, TEXT("GeneratedCodeHelpers"), true);
+
 		TSet<FString> AlreadyIncluded;
 		AlreadyIncluded.Add(SourceStruct->GetName());
-
 		auto EmitInner = [&](FStringOutputDevice& Dst, const TArray<UObject*>& Src)
 		{
+			auto EngineSourceDir = FPaths::EngineSourceDir();
+			auto GameSourceDir = FPaths::GameSourceDir();
+			auto CurrentPackage = SourceStruct->GetTypedOuter<UPackage>();
+
 			for (UObject* Obj : Src)
 			{
 				bool bWantedType = Obj && (Obj->IsA<UBlueprintGeneratedClass>() || Obj->IsA<UUserDefinedEnum>() || Obj->IsA<UUserDefinedStruct>());
@@ -451,21 +454,47 @@ void FBlueprintCompilerCppBackendBase::EmitFileBeginning(const FString& CleanNam
 					}
 				}
 
-				if (!bWantedType)
+				auto BPGC = (Obj && !bWantedType) ? Cast<UBlueprintGeneratedClass>(Obj->GetClass()) : nullptr;
+				if (BPGC)
 				{
-					Obj = Obj ? Cast<UBlueprintGeneratedClass>(Obj->GetClass()) : nullptr;
-					bWantedType = (nullptr != Obj);
+					bWantedType = true;
+					Obj = BPGC;
 				}
 
+				// Wanted no-native type, thet will be converted
 				if (bWantedType && Obj)
 				{
 					const FString Name = Obj->GetName();
-					if (!AlreadyIncluded.Contains(Name))
+					bool bAlreadyIncluded = false;
+					AlreadyIncluded.Add(Name, &bAlreadyIncluded);
+					if (!bAlreadyIncluded)
 					{
-						AlreadyIncluded.Add(Name);
-						Emit(Dst, *FString::Printf(TEXT("#include \"%s.h\"\n"), *Name));
+						EmitIncludeHeader(Dst, *Name, true);
 					}
 				}
+				// headers for native items
+				else if (auto Field = Cast<UField>(Obj))
+				{
+					FString PackPath;
+					if ((Field->GetTypedOuter<UPackage>() != CurrentPackage)
+						&& FSourceCodeNavigation::FindClassHeaderPath(Field, PackPath))
+					{
+						if (!PackPath.RemoveFromStart(EngineSourceDir))
+						{
+							if (!PackPath.RemoveFromStart(GameSourceDir))
+							{
+								PackPath = FPaths::GetCleanFilename(PackPath);
+							}
+						}
+						bool bAlreadyIncluded = false;
+						AlreadyIncluded.Add(PackPath, &bAlreadyIncluded);
+						if (!bAlreadyIncluded)
+						{
+							EmitIncludeHeader(Dst, *PackPath, false);
+						}
+					}
+
+				}
 			}
 			Emit(Dst, TEXT("\n"));
 		};
diff --git a/Engine/Source/Developer/BlueprintCompilerCppBackend/Private/BlueprintCompilerCppBackendUtils.cpp b/Engine/Source/Developer/BlueprintCompilerCppBackend/Private/BlueprintCompilerCppBackendUtils.cpp
--- a/Engine/Source/Developer/BlueprintCompilerCppBackend/Private/BlueprintCompilerCppBackendUtils.cpp
+++ b/Engine/Source/Developer/BlueprintCompilerCppBackend/Private/BlueprintCompilerCppBackendUtils.cpp
@@ -345,55 +345,6 @@ FString FEmitHelper::EmitLifetimeReplicatedPropsImpl(UClass* SourceClass, const
 	return Result;
 }
 
-FString FEmitHelper::GatherNativeHeadersToInclude(UField* SourceItem, const TArray<FString>& PersistentHeaders)
-{
-	TSet<FString> HeaderFiles;
-	HeaderFiles.Append(PersistentHeaders);
-	{
-		TArray<UObject*> ReferencedObjects;
-		{
-			FReferenceFinder ReferenceFinder(ReferencedObjects, NULL, false, false, false, true);
-			TArray<UObject*> ObjectsToCheck;
-			GetObjectsWithOuter(SourceItem, ObjectsToCheck, true);
-			//CDO?
-			for (auto Obj : ObjectsToCheck)
-			{
-				ReferenceFinder.FindReferences(Obj);
-			}
-		}
-
-		auto EngineSourceDir = FPaths::EngineSourceDir();
-		auto GameSourceDir = FPaths::GameSourceDir();
-		auto CurrentPackage = SourceItem->GetTypedOuter<UPackage>();
-		for (auto Obj : ReferencedObjects)
-		{
-			auto Field = Cast<UField>(Obj);
-			FString PackPath;
-			if (Field
-				&& (Field->GetTypedOuter<UPackage>() != CurrentPackage)
-				&& FSourceCodeNavigation::FindClassHeaderPath(Field, PackPath))
-			{
-				if (!PackPath.RemoveFromStart(EngineSourceDir))
-				{
-					if (!PackPath.RemoveFromStart(GameSourceDir))
-					{
-						PackPath = FPaths::GetCleanFilename(PackPath);
-					}
-				}
-				HeaderFiles.Add(PackPath);
-			}
-		}
-	}
-
-	FString Result;
-	for (auto HeaderFile : HeaderFiles)
-	{
-		Result += FString::Printf(TEXT("#include \"%s\"\n"), *HeaderFile);
-	}
-
-	return Result;
-}
-
 FString FEmitHelper::LiteralTerm(FEmitterLocalContext& EmitterContext, const FEdGraphPinType& Type, const FString& CustomValue, UObject* LiteralObject)
 {
 	auto Schema = GetDefault<UEdGraphSchema_K2>();
@@ -644,19 +595,40 @@ bool FEmitHelper::ShouldHandleAsImplementableEvent(UFunction* Function)
 	return false;
 }
 
-bool FEmitHelper::GenerateAssignmentCast(const FEdGraphPinType& LType, const FEdGraphPinType& RType, FString& OutCastBegin, FString& OutCastEnd)
+bool FEmitHelper::GenerateAutomaticCast(const FEdGraphPinType& LType, const FEdGraphPinType& RType, FString& OutCastBegin, FString& OutCastEnd)
 {
 	// BYTE to ENUM cast
-	auto LTypeEnum = Cast<UEnum>(LType.PinSubCategoryObject.Get());
-	if ((LType.PinCategory == UEdGraphSchema_K2::PC_Byte)
-		&& (RType.PinCategory == UEdGraphSchema_K2::PC_Byte)
-		&& !RType.PinSubCategoryObject.IsValid()
-		&& LTypeEnum)
-	{
-		FString EnumCppType = !LTypeEnum->CppType.IsEmpty() ? LTypeEnum->CppType : LTypeEnum->GetName();
-		OutCastBegin = FString::Printf(TEXT("static_cast<%s>("), *EnumCppType);
-		OutCastEnd = TEXT(")");
-		return true;
+	// ENUM to BYTE cast
+	if ((LType.PinCategory == UEdGraphSchema_K2::PC_Byte) && (RType.PinCategory == UEdGraphSchema_K2::PC_Byte))
+	{
+		auto LTypeEnum = Cast<UEnum>(LType.PinSubCategoryObject.Get());
+		auto RTypeEnum = Cast<UEnum>(RType.PinSubCategoryObject.Get());
+		if (!RTypeEnum && LTypeEnum)
+		{
+			FString EnumCppType = !LTypeEnum->CppType.IsEmpty() ? LTypeEnum->CppType : LTypeEnum->GetName();
+			OutCastBegin = FString::Printf(TEXT("static_cast<%s>("), *EnumCppType);
+			OutCastEnd = TEXT(")");
+			return true;
+		}
+		if (!LTypeEnum && RTypeEnum)
+		{
+			OutCastBegin = TEXT("EnumToByte(");
+			OutCastEnd = TEXT(")");
+			return true;
+		}
+	}
+
+	// OBJECT to OBJECT
+	if (LType.PinCategory == UEdGraphSchema_K2::PC_Object)
+	{
+		UClass* LClass = Cast<UClass>(LType.PinSubCategoryObject.Get());
+		UClass* RClass = Cast<UClass>(RType.PinSubCategoryObject.Get());
+		if (LClass && RClass && LClass->IsChildOf(RClass) && !RClass->IsChildOf(LClass))
+		{
+			OutCastBegin = FString::Printf(TEXT("CastChecked<%s%s>("), LClass->GetPrefixCPP(), *LClass->GetName());
+			OutCastEnd = TEXT(")");
+			return true;
+		}
 	}
 
 	return false;
diff --git a/Engine/Source/Developer/BlueprintCompilerCppBackend/Private/BlueprintCompilerCppBackendUtils.h b/Engine/Source/Developer/BlueprintCompilerCppBackend/Private/BlueprintCompilerCppBackendUtils.h
--- a/Engine/Source/Developer/BlueprintCompilerCppBackend/Private/BlueprintCompilerCppBackendUtils.h
+++ b/Engine/Source/Developer/BlueprintCompilerCppBackend/Private/BlueprintCompilerCppBackendUtils.h
@@ -138,8 +138,10 @@ public:
 
 		if (auto UDE = Cast<UEnum>(Object))
 		{
+			// TODO:
 			// TODO: check if supported 
-			return FString::Printf(TEXT("%s_StaticEnum()"), *UDE->GetName());
+			//return FString::Printf(TEXT("%s_StaticEnum()"), *UDE->GetName());
+			return FString::Printf(TEXT("FindObjectChecked<UEnum>(ANY_PACKAGE, TEXT(\"%s\"))"), *UDE->GetName());
 		}
 
 		// TODO: handle subobjects
@@ -226,8 +228,6 @@ struct FEmitHelper
 
 	static FString EmitLifetimeReplicatedPropsImpl(UClass* SourceClass, const FString& CppClassName, const TCHAR* InCurrentIndent);
 
-	static FString GatherNativeHeadersToInclude(UField* SourceItem, const TArray<FString>& PersistentHeaders);
-
 	static FString LiteralTerm(FEmitterLocalContext& EmitterContext, const FEdGraphPinType& Type, const FString& CustomValue, UObject* LiteralObject);
 
 	static FString DefaultValue(FEmitterLocalContext& EmitterContext, const FEdGraphPinType& Type);
@@ -238,7 +238,7 @@ struct FEmitHelper
 
 	static bool ShouldHandleAsImplementableEvent(UFunction* Function);
 
-	static bool GenerateAssignmentCast(const FEdGraphPinType& LType, const FEdGraphPinType& RType, FString& OutCastBegin, FString& OutCastEnd);
+	static bool GenerateAutomaticCast(const FEdGraphPinType& LType, const FEdGraphPinType& RType, FString& OutCastBegin, FString& OutCastEnd);
 
 	static FString GenerateReplaceConvertedMD(UObject* Obj);
 };
diff --git a/Engine/Source/Runtime/Engine/Public/GeneratedCodeHelpers.h b/Engine/Source/Runtime/Engine/Public/GeneratedCodeHelpers.h
--- a/Engine/Source/Runtime/Engine/Public/GeneratedCodeHelpers.h
+++ b/Engine/Source/Runtime/Engine/Public/GeneratedCodeHelpers.h
@@ -23,6 +23,12 @@ FORCEINLINE bool IsValid(const FScriptInterface& Test)
 	return IsValid(Test.GetObject()) && (nullptr != Test.GetInterface());
 }
 
+template<class TEnum>
+FORCEINLINE uint8 EnumToByte(TEnumAsByte<TEnum> Val)
+{
+	return static_cast<uint8>(Val.GetValue());
+}
+
 struct FCustomThunkTemplates
 {
 private:
