UEBP-40 C++ code generation (WIP):
- Get Class Defaults is handled

diff --git a/Engine/Source/Developer/BlueprintCompilerCppBackend/Private/BlueprintCompilerCppBackend.cpp b/Engine/Source/Developer/BlueprintCompilerCppBackend/Private/BlueprintCompilerCppBackend.cpp
--- a/Engine/Source/Developer/BlueprintCompilerCppBackend/Private/BlueprintCompilerCppBackend.cpp
+++ b/Engine/Source/Developer/BlueprintCompilerCppBackend/Private/BlueprintCompilerCppBackend.cpp
@@ -25,7 +25,6 @@ public:
 		{
 			if (!Term->IsStructContextType() && (Term->Type.PinSubCategory != TEXT("self")))
 			{
-				ensure(!Term->bIsLiteral);
 				SafetyConditions.Add(CppBackend.TermToText(EmitterContext, Term, false));
 			}
 		}
@@ -505,8 +504,26 @@ FString FBlueprintCompilerCppBackend::TermToText(FEmitterLocalContext& EmitterCo
 		FString ResultPath(TEXT(""));
 		if ((Term->Context != nullptr) && (Term->Context->Name != PSC_Self))
 		{
-			ResultPath = TermToText(EmitterContext, Term->Context, false);
+			const bool bFromDefaultValue = Term->Context->IsClassContextType();
+			if (bFromDefaultValue)
+			{
+				const UClass* MinimalClass = ensure(Term->AssociatedVarProperty)
+					? Term->AssociatedVarProperty->GetOwnerClass()
+					: Cast<UClass>(Term->Context->Type.PinSubCategoryObject.Get());
+				if (MinimalClass)
+				{
+					ResultPath += FString::Printf(TEXT("GetDefaultValueSafe<%s%s>(")
+						, MinimalClass->GetPrefixCPP()
+						, *MinimalClass->GetName());
+				}
+				else
+				{
+					UE_LOG(LogK2Compiler, Error, TEXT("C++ backend cannot find specific class"));
+				}
+			}
 
+			ResultPath += TermToText(EmitterContext, Term->Context, false);
+			ResultPath += bFromDefaultValue ? TEXT(")") : TEXT("");
 			if (Term->Context->IsStructContextType())
 			{
 				ResultPath += TEXT(".");
diff --git a/Engine/Source/Runtime/Engine/Public/GeneratedCodeHelpers.h b/Engine/Source/Runtime/Engine/Public/GeneratedCodeHelpers.h
--- a/Engine/Source/Runtime/Engine/Public/GeneratedCodeHelpers.h
+++ b/Engine/Source/Runtime/Engine/Public/GeneratedCodeHelpers.h
@@ -13,22 +13,28 @@
 
 #include "Stack.h"
 
-FORCEINLINE UClass* DynamicMetaCast(const UClass* DesiredClass, UClass* SourceClass)
+inline UClass* DynamicMetaCast(const UClass* DesiredClass, UClass* SourceClass)
 {
 	return ((SourceClass)->IsChildOf(DesiredClass)) ? SourceClass : NULL;
 }
 
-FORCEINLINE bool IsValid(const FScriptInterface& Test)
+inline bool IsValid(const FScriptInterface& Test)
 {
 	return IsValid(Test.GetObject()) && (nullptr != Test.GetInterface());
 }
 
 template<class TEnum>
-FORCEINLINE uint8 EnumToByte(TEnumAsByte<TEnum> Val)
+inline uint8 EnumToByte(TEnumAsByte<TEnum> Val)
 {
 	return static_cast<uint8>(Val.GetValue());
 }
 
+template<class T>
+inline const T* GetDefaultValueSafe(UClass* Class)
+{
+	return IsValid(Class) ? GetDefault<T>(Class) : nullptr;
+}
+
 struct FCustomThunkTemplates
 {
 private:
