Fix for EnumCppType failing to compile when passed an enum literal. Template type could not be deduced because of usual int/enum issues.

diff --git a/Engine/Source/Developer/BlueprintCompilerCppBackend/Private/BlueprintCompilerCppBackendUtils.cpp b/Engine/Source/Developer/BlueprintCompilerCppBackend/Private/BlueprintCompilerCppBackendUtils.cpp
--- a/Engine/Source/Developer/BlueprintCompilerCppBackend/Private/BlueprintCompilerCppBackendUtils.cpp
+++ b/Engine/Source/Developer/BlueprintCompilerCppBackend/Private/BlueprintCompilerCppBackendUtils.cpp
@@ -630,7 +630,8 @@ bool FEmitHelper::GenerateAutomaticCast(const FEdGraphPinType& LType, const FEdG
 		}
 		if (!LTypeEnum && RTypeEnum)
 		{
-			OutCastBegin = TEXT("EnumToByte(");
+			FString EnumCppType = !RTypeEnum->CppType.IsEmpty() ? RTypeEnum->CppType : RTypeEnum->GetName();
+			OutCastBegin = FString::Printf(TEXT("EnumToByte<%s>("), *EnumCppType); 
 			OutCastEnd = TEXT(")");
 			return true;
 		}
