[UEBP-40] BP->C++ code generation (WiP) - First round of component-specific generated runtime code changes.

change summary:
- revised native class ctor init code gen for instanced default subobjects (DSOs)
- revised native class ctor init code gen for non-native components (SCS nodes)
- added support to bypass native class ctor init code gen for editor-only subobjects
- added support to fall back to creating a dummy UActorComponent/USceneComponent in the case of an editor-only subobject dependency
- added support for generating native class ctor code to set up the RootComponent property when there is not already a native root
- added support for populating the BlueprintCreatedComponents array at Actor construction time for generated BP native classes
- fixed some compile issues after generating native code and attempting to rebuild project code
- fixed a runtime issue with generated native ctor init code for UCS-added components

diff --git a/Engine/Source/Developer/BlueprintCompilerCppBackend/Private/BlueprintCompilerCppBackendUtils.h b/Engine/Source/Developer/BlueprintCompilerCppBackend/Private/BlueprintCompilerCppBackendUtils.h
--- a/Engine/Source/Developer/BlueprintCompilerCppBackend/Private/BlueprintCompilerCppBackendUtils.h
+++ b/Engine/Source/Developer/BlueprintCompilerCppBackend/Private/BlueprintCompilerCppBackendUtils.h
@@ -267,11 +267,11 @@ private:
 	// returns empty string if cannot handle
 	static FString HandleSpecialTypes(FEmitterLocalContext& Context, const UProperty* Property, const uint8* ValuePtr);
 
-	static UActorComponent* HandleNonNativeComponent(FEmitterLocalContext& Context, UBlueprintGeneratedClass* BPGC, FName ObjectName, bool bNew, const FString& NativeName);
+	static FString HandleNonNativeComponent(FEmitterLocalContext& Context, const USCS_Node* Node, TSet<const UProperty*>& OutHandledProperties, const USCS_Node* ParentNode = nullptr);
 	
 	// Creates the subobject (of class) returns it's native local name, 
 	// returns empty string if cannot handle
 	static FString HandleClassSubobject(FEmitterLocalContext& Context, UObject* Object);
 
-	static FString HandleInstancedSubobject(FEmitterLocalContext& Context, UObject* Object);
+	static FString HandleInstancedSubobject(FEmitterLocalContext& Context, UObject* Object, bool bSkipEditorOnlyCheck = false);
 };
\ No newline at end of file
diff --git a/Engine/Source/Developer/BlueprintCompilerCppBackend/Private/BlueprintCompilerCppBackendValueHelper.cpp b/Engine/Source/Developer/BlueprintCompilerCppBackend/Private/BlueprintCompilerCppBackendValueHelper.cpp
--- a/Engine/Source/Developer/BlueprintCompilerCppBackend/Private/BlueprintCompilerCppBackendValueHelper.cpp
+++ b/Engine/Source/Developer/BlueprintCompilerCppBackend/Private/BlueprintCompilerCppBackendValueHelper.cpp
@@ -14,6 +14,25 @@ void FEmitDefaultValueHelper::OuterGenerate(FEmitterLocalContext& Context
 	, EPropertyAccessOperator AccessOperator
 	, bool bAllowProtected)
 {
+	// Determine if the given property contains an instanced default subobject reference. We only get here if the values are not identical.
+	auto IsInstancedSubobjectLambda = [&](int32 ArrayIndex) -> bool
+	{
+		if (auto ObjectProperty = Cast<UObjectProperty>(Property))
+		{
+			check(DataContainer);
+			check(OptionalDefaultDataContainer);
+
+			auto ObjectPropertyValue = ObjectProperty->GetObjectPropertyValue_InContainer(DataContainer, ArrayIndex);
+			auto DefaultObjectPropertyValue = ObjectProperty->GetObjectPropertyValue_InContainer(OptionalDefaultDataContainer, ArrayIndex);
+			if (ObjectPropertyValue && ObjectPropertyValue->IsDefaultSubobject() && DefaultObjectPropertyValue && DefaultObjectPropertyValue->IsDefaultSubobject() && ObjectPropertyValue->GetFName() == DefaultObjectPropertyValue->GetFName())
+			{
+				return true;
+			}
+		}
+
+		return false;
+	};
+
 	if (Property->HasAnyPropertyFlags(CPF_EditorOnly | CPF_Transient))
 	{
 		UE_LOG(LogK2Compiler, Log, TEXT("FEmitDefaultValueHelper Skip EditorOnly or Transient property: %s"), *Property->GetPathName());
@@ -22,7 +41,8 @@ void FEmitDefaultValueHelper::OuterGenerate(FEmitterLocalContext& Context
 
 	for (int32 ArrayIndex = 0; ArrayIndex < Property->ArrayDim; ++ArrayIndex)
 	{
-		if (!OptionalDefaultDataContainer || !Property->Identical_InContainer(DataContainer, OptionalDefaultDataContainer, ArrayIndex))
+		if (!OptionalDefaultDataContainer
+			|| (!Property->Identical_InContainer(DataContainer, OptionalDefaultDataContainer, ArrayIndex) && !IsInstancedSubobjectLambda(ArrayIndex)))
 		{
 			FString PathToMember;
 			if (Property->HasAnyPropertyFlags(CPF_NativeAccessSpecifierPrivate) || (!bAllowProtected && Property->HasAnyPropertyFlags(CPF_NativeAccessSpecifierProtected)))
@@ -232,92 +252,75 @@ FString FEmitDefaultValueHelper::HandleSpecialTypes(FEmitterLocalContext& Contex
 	return FString();
 }
 
-UActorComponent* FEmitDefaultValueHelper::HandleNonNativeComponent(FEmitterLocalContext& Context, UBlueprintGeneratedClass* BPGC, FName ObjectName, bool bNew, const FString& NativeName)
+FString FEmitDefaultValueHelper::HandleNonNativeComponent(FEmitterLocalContext& Context, const USCS_Node* Node, TSet<const UProperty*>& OutHandledProperties, const USCS_Node* ParentNode)
 {
-	USCS_Node* Node = nullptr;
-	for (UBlueprintGeneratedClass* NodeOwner = BPGC; NodeOwner && !Node; NodeOwner = Cast<UBlueprintGeneratedClass>(NodeOwner->GetSuperClass()))
-	{
-		Node = NodeOwner->SimpleConstructionScript ? NodeOwner->SimpleConstructionScript->FindSCSNode(ObjectName) : nullptr;
-	}
+	check(Node);
 
-	if (!Node)
+	FString VariableName;
+	UBlueprintGeneratedClass* BPGC = CastChecked<UBlueprintGeneratedClass>(Context.GetCurrentlyGeneratedClass());
+	if (UActorComponent* ComponentTemplate = Node->GetActualComponentTemplate(BPGC))
 	{
-		return nullptr;
-	}
-
-	ensure(bNew == (Node->GetOuter() == BPGC->SimpleConstructionScript));
-	UActorComponent* ComponentTemplate = bNew
-		? Node->ComponentTemplate
-		: (BPGC->InheritableComponentHandler ? BPGC->InheritableComponentHandler->GetOverridenComponentTemplate(Node) : nullptr);
+		VariableName = Node->VariableName.ToString();
+		Context.AddObjectFromLocalProperty_InConstructor(ComponentTemplate, VariableName);
 
-	if (!ComponentTemplate)
-	{
-		return nullptr;
-	}
-
-	auto ComponentClass = ComponentTemplate->GetClass();
-	if (bNew)
-	{
-		Context.AddHighPriorityLine(FString::Printf(TEXT("%s = CreateDefaultSubobject<%s%s>(TEXT(\"%s\"));")
-			, *NativeName, ComponentClass->GetPrefixCPP(), *ComponentClass->GetName(), *ComponentTemplate->GetName()));
-		Context.AddObjectFromLocalProperty_InConstructor(ComponentTemplate, NativeName);
-	}
-
-	UObject* ParentTemplateComponent = bNew 
-		? nullptr 
-		: Node->GetActualComponentTemplate(Cast<UBlueprintGeneratedClass>(BPGC->GetSuperClass()));
-	ensure(bNew || ParentTemplateComponent);
-	ensure(ParentTemplateComponent != ComponentTemplate);
-	ensure(!ParentTemplateComponent || (ParentTemplateComponent->GetClass() == ComponentClass));
-
-	const UObject* ObjectToCompare = bNew
-		? ComponentClass->GetDefaultObject(false)
-		: ParentTemplateComponent;
-	for (auto Property : TFieldRange<const UProperty>(ComponentClass))
-	{
-		OuterGenerate(Context, Property, NativeName
-			, reinterpret_cast<const uint8*>(ComponentTemplate)
-			, reinterpret_cast<const uint8*>(ObjectToCompare)
-			, EPropertyAccessOperator::Pointer);
-	}
-
-	auto FindParentComponentName = [&]() -> FString
-	{
-		FString ResultName;
-		if (auto ParentNode = Node->GetSCS()->FindParentNode(Node))
+		const UObjectProperty* VariableProperty = FindField<UObjectProperty>(BPGC, *VariableName);
+		if (VariableProperty)
 		{
-			const FName ParentVariableName = ParentNode->GetVariableName();
-			ResultName = (ParentVariableName != NAME_None) ? ParentVariableName.ToString() : FString();
+			OutHandledProperties.Add(VariableProperty);
 		}
 
-		if (ResultName.IsEmpty())
+		if (ComponentTemplate->GetOuter() == BPGC)
 		{
-			auto ParentTemplate = Node->GetParentComponentTemplate(CastChecked<UBlueprint>(BPGC->ClassGeneratedBy));
-			FString ParentComponentNativeName;
-			if (ParentTemplate && Context.FindLocalObject_InConstructor(ParentTemplateComponent, ParentComponentNativeName))
+			UClass* ComponentClass = ComponentTemplate->GetClass();
+			check(ComponentClass != nullptr);
+
+			UObject* ObjectToCompare = ComponentClass->GetDefaultObject(false);
+
+			if (ComponentTemplate->HasAnyFlags(RF_InheritableComponentTemplate))
 			{
-				ResultName = ParentComponentNativeName;
+				ObjectToCompare = Node->GetActualComponentTemplate(Cast<UBlueprintGeneratedClass>(BPGC->GetSuperClass()));
 			}
-		}
+			else
+			{
+				Context.AddHighPriorityLine(FString::Printf(TEXT("%s%s = CreateDefaultSubobject<%s%s>(TEXT(\"%s\"));")
+					, VariableProperty == nullptr ? TEXT("auto ") : TEXT(""), *VariableName
+					, ComponentClass->GetPrefixCPP(), *ComponentClass->GetName(), *ComponentTemplate->GetName()));
 
-		if (ResultName.IsEmpty())
-		{
-			ResultName = (Node->ParentComponentOrVariableName != NAME_None) ? Node->ParentComponentOrVariableName.ToString() : FString();
-		}
+				Context.AddLine(FString::Printf(TEXT("%s->CreationMethod = EComponentCreationMethod::SimpleConstructionScript;"), *VariableName));
 
-		return ResultName;
-	};
+				FString ParentVariableName;
+				if (ParentNode)
+				{
+					ParentVariableName = ParentNode->VariableName.ToString();
+				}
+				else if (USceneComponent* ParentComponentTemplate = Node->GetParentComponentTemplate(CastChecked<UBlueprint>(BPGC->ClassGeneratedBy)))
+				{
+					Context.FindLocalObject_InConstructor(ParentComponentTemplate, ParentVariableName);
+				}
 
-	const FString ParentNodeName = FindParentComponentName();
-	if (!ParentNodeName.IsEmpty())
-	{
-		Context.AddLine(FString::Printf(TEXT("%s->AttachParent = %s;"), *NativeName, *ParentNodeName));
+				if (!ParentVariableName.IsEmpty())
+				{
+					Context.AddLine(FString::Printf(TEXT("%s->AttachParent = %s;"), *VariableName, *ParentVariableName));
+				}
+			}
+
+			for (auto Property : TFieldRange<const UProperty>(ComponentClass))
+			{
+				OuterGenerate(Context, Property, VariableName
+					, reinterpret_cast<const uint8*>(ComponentTemplate)
+					, reinterpret_cast<const uint8*>(ObjectToCompare)
+					, EPropertyAccessOperator::Pointer);
+			}
+		}
 	}
 
-	// TODO: Set root component
-	// TODO: Update local transform, etc..
+	// Recursively handle child nodes.
+	for (auto ChildNode : Node->ChildNodes)
+	{
+		HandleNonNativeComponent(Context, ChildNode, OutHandledProperties, Node);
+	}
 
-	return ComponentTemplate;
+	return VariableName;
 }
 
 FString FEmitDefaultValueHelper::GenerateConstructor(UClass* InBPGC)
@@ -396,10 +399,11 @@ FString FEmitDefaultValueHelper::GenerateConstructor(UClass* InBPGC)
 			FString NativeName;
 			const bool bFound = Context.FindLocalObject_InConstructor(Obj, NativeName);
 			ensure(bFound);
-			Context.AddHighPriorityLine(FString::Printf(TEXT("auto %s = CastChecked<%s%s>(GetClass()->ConvertedSubobjectsFromBPGC[%d]);")
+			Context.AddHighPriorityLine(FString::Printf(TEXT("auto %s = CastChecked<%s%s>(%s::StaticClass()->ConvertedSubobjectsFromBPGC[%d]);")
 				, *NativeName
 				, Obj->GetClass()->GetPrefixCPP()
 				, *Obj->GetClass()->GetName()
+				, *CppClassName
 				, SubobjectIndex));
 			SubobjectIndex++;
 		}
@@ -412,31 +416,83 @@ FString FEmitDefaultValueHelper::GenerateConstructor(UClass* InBPGC)
 	check(CDO && ParentCDO);
 	Context.AddLine(TEXT(""));
 
-	// TIMELINES
 	TSet<const UProperty*> HandledProperties;
-	const bool bIsActor = BPGC->IsChildOf<AActor>();
 
-	// COMPONENTS
-	for (auto Property : TFieldRange<const UProperty>(BPGC))
+	// Generate ctor init code for native class default subobjects that are always instanced (e.g. components).
+	// @TODO (pkavan) - We can probably make this faster by generating code to index through the DSO array instead (i.e. in place of HandleInstancedSubobject which will generate a lookup call per DSO).
+	TArray<UObject*> NativeDefaultObjectSubobjects;
+	BPGC->GetDefaultObjectSubobjects(NativeDefaultObjectSubobjects);
+	for (auto DSO : NativeDefaultObjectSubobjects)
 	{
-		const bool bNewProperty = Property->GetOwnerStruct() == BPGC;
-		const bool bIsAccessible = bNewProperty || !Property->HasAnyPropertyFlags(CPF_NativeAccessSpecifierPrivate);
-		if (bIsAccessible && !HandledProperties.Contains(Property))
+		if (DSO && DSO->GetClass()->HasAnyClassFlags(CLASS_DefaultToInstanced))
 		{
-			auto ObjectProperty = Cast<UObjectProperty>(Property);
-			const bool bComponentProp = ObjectProperty && ObjectProperty->PropertyClass && ObjectProperty->PropertyClass->IsChildOf<UActorComponent>();
-			const bool bNullValue = ObjectProperty && (nullptr == ObjectProperty->GetPropertyValue_InContainer(CDO));
-			const bool bActorComponent = bIsActor && bComponentProp && bNullValue;
-			if (bIsActor && bComponentProp && bNullValue)
+			// Determine if this is an editor-only subobject.
+			bool bIsEditorOnlySubobject = false;
+			if (const UActorComponent* ActorComponent = Cast<UActorComponent>(DSO))
 			{
-				auto HandledComponent = HandleNonNativeComponent(Context, BPGC, Property->GetFName(), bNewProperty, Property->GetNameCPP());
-				if (HandledComponent)
+				bIsEditorOnlySubobject = ActorComponent->IsEditorOnly();
+			}
+
+			// Skip ctor code gen for editor-only subobjects, since they won't be used by the runtime. Any dependencies on editor-only subobjects will be handled later (see HandleInstancedSubobject).
+			if (!bIsEditorOnlySubobject)
+			{
+				HandleInstancedSubobject(Context, DSO, true);
+			}
+		}
+	}
+
+	// Check for a valid RootComponent property value; mark it as handled if already set in the defaults.
+	bool bNeedsRootComponentAssignment = false;
+	static const FName RootComponentPropertyName(TEXT("RootComponent"));
+	const UObjectProperty* RootComponentProperty = FindField<UObjectProperty>(BPGC, RootComponentPropertyName);
+	if (RootComponentProperty)
+	{
+		if (RootComponentProperty->GetObjectPropertyValue_InContainer(CDO))
+		{
+			HandledProperties.Add(RootComponentProperty);
+		}
+		else
+		{
+			bNeedsRootComponentAssignment = true;
+		}
+	}
+
+	// Generate ctor init code for the SCS node hierarchy (i.e. non-native components). SCS nodes may have dependencies on native DSOs, but not vice-versa.
+	TArray<const UBlueprintGeneratedClass*> BPGCStack;
+	const bool bErrorFree = UBlueprintGeneratedClass::GetGeneratedClassesHierarchy(BPGC, BPGCStack);
+	if (bErrorFree)
+	{
+		// Start at the base of the hierarchy so that dependencies are handled first.
+		for (int32 i = BPGCStack.Num() - 1; i >= 0; --i)
+		{
+			if (BPGCStack[i]->SimpleConstructionScript)
+			{
+				for (auto Node : BPGCStack[i]->SimpleConstructionScript->GetRootNodes())
 				{
-					ensure(!ActorComponentTempatesOwnedByClass.Contains(HandledComponent));
-					continue;
+					if (Node)
+					{
+						const FString VariableName = HandleNonNativeComponent(Context, Node, HandledProperties);
+
+						if (i == 0 && bNeedsRootComponentAssignment && Node->ComponentTemplate && Node->ComponentTemplate->IsA<USceneComponent>() && !VariableName.IsEmpty())
+						{
+							Context.AddLine(FString::Printf(TEXT("RootComponent = %s;"), *VariableName));
+
+							bNeedsRootComponentAssignment = false;
+							HandledProperties.Add(RootComponentProperty);
+						}
+					}
 				}
 			}
+		}
+	}
 
+	// Generate ctor init code for generated Blueprint class property values that may differ from parent class defaults (or that otherwise belong to the generated Blueprint class).
+	for (auto Property : TFieldRange<const UProperty>(BPGC))
+	{
+		const bool bNewProperty = Property->GetOwnerStruct() == BPGC;
+		const bool bIsAccessible = bNewProperty || !Property->HasAnyPropertyFlags(CPF_NativeAccessSpecifierPrivate);
+		if (bIsAccessible && !HandledProperties.Contains(Property))
+		{
 			OuterGenerate(Context, Property, TEXT(""), reinterpret_cast<const uint8*>(CDO), bNewProperty ? nullptr : reinterpret_cast<const uint8*>(ParentCDO), EPropertyAccessOperator::None, true);
 		}
 	}
@@ -481,19 +537,53 @@ FString FEmitDefaultValueHelper::HandleClassSubobject(FEmitterLocalContext& Cont
 	return LocalNativeName;
 }
 
-FString FEmitDefaultValueHelper::HandleInstancedSubobject(FEmitterLocalContext& Context, UObject* Object)
+FString FEmitDefaultValueHelper::HandleInstancedSubobject(FEmitterLocalContext& Context, UObject* Object, bool bSkipEditorOnlyCheck)
 {
 	check(Object);
 
-	const FString LocalNativeName = Context.AddNewObject_InConstructor(Object);
+	// Make sure we don't emit initialization code for the same object more than once.
+	FString LocalNativeName;
+	if (Context.FindLocalObject_InConstructor(Object, LocalNativeName))
+	{
+		return LocalNativeName;
+	}
+	else
+	{
+		LocalNativeName = Context.AddNewObject_InConstructor(Object);
+	}
+
 	UClass* ObjectClass = Object->GetClass();
 
+	// Determine if this is an editor-only subobject. When handling as a dependency, we'll create a "dummy" object in its place (below).
+	bool bIsEditorOnlySubobject = false;
+	if (!bSkipEditorOnlyCheck)
+	{
+		if (UActorComponent* ActorComponent = Cast<UActorComponent>(Object))
+		{
+			bIsEditorOnlySubobject = ActorComponent->IsEditorOnly();
+			if (bIsEditorOnlySubobject)
+			{
+				// Replace the potentially editor-only class with a base actor/scene component class that's available to the runtime. We'll create a "dummy" object of this type to stand in for the editor-only subobject below.
+				ObjectClass = ObjectClass->IsChildOf<USceneComponent>() ? USceneComponent::StaticClass() : UActorComponent::StaticClass();
+			}
+		}
+	}
+
 	auto BPGC = Context.GetCurrentlyGeneratedClass();
 	auto CDO = BPGC ? BPGC->GetDefaultObject(false) : nullptr;
-	if (ensure(CDO) && (CDO == Object->GetOuter()))
+	if (!bIsEditorOnlySubobject && ensure(CDO) && (CDO == Object->GetOuter()))
 	{
-		Context.AddHighPriorityLine(FString::Printf(TEXT("auto %s = CreateDefaultSubobject<%s%s>(TEXT(\"%s\"));")
+		Context.AddHighPriorityLine(FString::Printf(TEXT("auto %s = CastChecked<%s%s>(GetDefaultSubobjectByName(TEXT(\"%s\")));")
 			, *LocalNativeName, ObjectClass->GetPrefixCPP(), *ObjectClass->GetName(), *Object->GetName()));
+
+		const UObject* ObjectArchetype = Object->GetArchetype();
+		for (auto Property : TFieldRange<const UProperty>(ObjectClass))
+		{
+			OuterGenerate(Context, Property, LocalNativeName
+				, reinterpret_cast<const uint8*>(Object)
+				, reinterpret_cast<const uint8*>(ObjectArchetype)
+				, EPropertyAccessOperator::Pointer);
+		}
 	}
 	else
 	{
@@ -507,13 +597,5 @@ FString FEmitDefaultValueHelper::HandleInstancedSubobject(FEmitterLocalContext&
 			, *OuterStr, *Object->GetName()));
 	}
 
-	for (auto Property : TFieldRange<const UProperty>(ObjectClass))
-	{
-		OuterGenerate(Context, Property, LocalNativeName
-			, reinterpret_cast<const uint8*>(Object)
-			, reinterpret_cast<const uint8*>(ObjectClass->GetDefaultObject(false))
-			, EPropertyAccessOperator::Pointer);
-	}
-
 	return LocalNativeName;
 }
\ No newline at end of file
diff --git a/Engine/Source/Runtime/Engine/Classes/Engine/InheritableComponentHandler.h b/Engine/Source/Runtime/Engine/Classes/Engine/InheritableComponentHandler.h
--- a/Engine/Source/Runtime/Engine/Classes/Engine/InheritableComponentHandler.h
+++ b/Engine/Source/Runtime/Engine/Classes/Engine/InheritableComponentHandler.h
@@ -18,7 +18,7 @@ struct ENGINE_API FComponentKey
 		: OwnerClass(nullptr)
 	{}
 
-	FComponentKey(USCS_Node* SCSNode);
+	FComponentKey(const USCS_Node* SCSNode);
 #if WITH_EDITOR
 	FComponentKey(UBlueprint* Blueprint, const FUCSComponentId& UCSComponentID);
 #endif 
@@ -44,13 +44,13 @@ struct ENGINE_API FComponentKey
 	UActorComponent* GetOriginalTemplate() const;
 	bool RefreshVariableName();
 
-	UClass* GetComponentOwner()  const { return OwnerClass; }
+	class UBlueprintGeneratedClass* GetComponentOwner()  const { return OwnerClass; }
 	FName   GetSCSVariableName() const { return SCSVariableName; }
 	FGuid   GetAssociatedGuid()  const { return AssociatedGuid; }
 
 private: 
 	UPROPERTY()
-	UBlueprintGeneratedClass* OwnerClass;
+	class UBlueprintGeneratedClass* OwnerClass;
 
 	UPROPERTY()
 	FName SCSVariableName;
diff --git a/Engine/Source/Runtime/Engine/Classes/Engine/SCS_Node.h b/Engine/Source/Runtime/Engine/Classes/Engine/SCS_Node.h
--- a/Engine/Source/Runtime/Engine/Classes/Engine/SCS_Node.h
+++ b/Engine/Source/Runtime/Engine/Classes/Engine/SCS_Node.h
@@ -95,7 +95,7 @@ class USCS_Node : public UObject
 	UActorComponent* ExecuteNodeOnActor(AActor* Actor, USceneComponent* ParentComponent, const FTransform* RootTransform, bool bIsDefaultTransform);
 
 	/** Return the actual component template used in the BPGC. The template can be overridden in a child. */
-	ENGINE_API UActorComponent* GetActualComponentTemplate(UBlueprintGeneratedClass* ActualBPGC);
+	ENGINE_API UActorComponent* GetActualComponentTemplate(class UBlueprintGeneratedClass* ActualBPGC) const;
 
 	/** Returns an array containing this node and all children below it */
 	TArray<USCS_Node*> GetAllNodes();
diff --git a/Engine/Source/Runtime/Engine/Private/BlueprintGeneratedClass.cpp b/Engine/Source/Runtime/Engine/Private/BlueprintGeneratedClass.cpp
--- a/Engine/Source/Runtime/Engine/Private/BlueprintGeneratedClass.cpp
+++ b/Engine/Source/Runtime/Engine/Private/BlueprintGeneratedClass.cpp
@@ -570,6 +570,16 @@ void UBlueprintGeneratedClass::CreateComponentsForActor(const UClass* ThisClass,
 	}
 	else
 	{
+		TInlineComponentArray<UActorComponent*> Components;
+		Actor->GetComponents(Components);
+		for (auto Component : Components)
+		{
+			if (Component->IsCreatedByConstructionScript())
+			{
+				Actor->BlueprintCreatedComponents.Add(Component);
+			}
+		}
+
 		for (auto MiscObj : ThisClass->ConvertedSubobjectsFromBPGC)
 		{
 			auto TimelineTemplate = Cast<const UTimelineTemplate>(MiscObj);
diff --git a/Engine/Source/Runtime/Engine/Private/InheritableComponentHandler.cpp b/Engine/Source/Runtime/Engine/Private/InheritableComponentHandler.cpp
--- a/Engine/Source/Runtime/Engine/Private/InheritableComponentHandler.cpp
+++ b/Engine/Source/Runtime/Engine/Private/InheritableComponentHandler.cpp
@@ -347,11 +347,11 @@ const FComponentOverrideRecord* UInheritableComponentHandler::FindRecord(const F
 
 // FComponentOverrideRecord
 
-FComponentKey::FComponentKey(USCS_Node* SCSNode) : OwnerClass(nullptr)
+FComponentKey::FComponentKey(const USCS_Node* SCSNode) : OwnerClass(nullptr)
 {
 	if (SCSNode)
 	{
-		USimpleConstructionScript* ParentSCS = SCSNode->GetSCS();
+		const USimpleConstructionScript* ParentSCS = SCSNode->GetSCS();
 		OwnerClass      = ParentSCS ? Cast<UBlueprintGeneratedClass>(ParentSCS->GetOwnerClass()) : nullptr;
 		AssociatedGuid  = SCSNode->VariableGuid;
 		SCSVariableName = SCSNode->GetVariableName();
diff --git a/Engine/Source/Runtime/Engine/Private/SCS_Node.cpp b/Engine/Source/Runtime/Engine/Private/SCS_Node.cpp
--- a/Engine/Source/Runtime/Engine/Private/SCS_Node.cpp
+++ b/Engine/Source/Runtime/Engine/Private/SCS_Node.cpp
@@ -26,7 +26,7 @@ USCS_Node::USCS_Node(const FObjectInitializer& ObjectInitializer)
 #endif
 }
 
-UActorComponent* USCS_Node::GetActualComponentTemplate(UBlueprintGeneratedClass* ActualBPGC)
+UActorComponent* USCS_Node::GetActualComponentTemplate(UBlueprintGeneratedClass* ActualBPGC) const
 {
 	UActorComponent* OverridenComponentTemplate = nullptr;
 	static const FBoolConfigValueHelper EnableInheritableComponents(TEXT("Kismet"), TEXT("bEnableInheritableComponents"), GEngineIni);
@@ -391,7 +391,7 @@ USceneComponent* USCS_Node::GetParentComponentTemplate(UBlueprint* InBlueprint)
 						if(CompTemplate != NULL && ParentSCSNodes[ParentNodeIndex]->VariableName == ParentComponentOrVariableName)
 						{
 							// Found a match; this is our parent, we're done
-							ParentComponentTemplate = CompTemplate;
+							ParentComponentTemplate = Cast<USceneComponent>(ParentSCSNodes[ParentNodeIndex]->GetActualComponentTemplate(Cast<UBlueprintGeneratedClass>(InBlueprint->GeneratedClass)));
 							break;
 						}
 					}
