UEBP-40 C++ code generation (WIP):
- Dependencies from CDO are collected

diff --git a/Engine/Source/Developer/BlueprintCompilerCppBackend/Private/BlueprintCompilerCppBackendBase.cpp b/Engine/Source/Developer/BlueprintCompilerCppBackend/Private/BlueprintCompilerCppBackendBase.cpp
--- a/Engine/Source/Developer/BlueprintCompilerCppBackend/Private/BlueprintCompilerCppBackendBase.cpp
+++ b/Engine/Source/Developer/BlueprintCompilerCppBackend/Private/BlueprintCompilerCppBackendBase.cpp
@@ -408,6 +408,11 @@ void FBlueprintCompilerCppBackendBase::EmitFileBeginning(const FString& CleanNam
 				}
 			}
 
+			if (auto BPGC = Cast<UClass>(SourceStruct))
+			{
+				BodyReferenceFinder.FindReferences(BPGC->GetDefaultObject(false));
+			}
+
 			if (auto SuperStruct = SourceStruct->GetSuperStruct())
 			{
 				IncludeInHeader.AddUnique(SuperStruct);
diff --git a/Engine/Source/Editor/Kismet/Private/NativeCodeGenerationTool.cpp b/Engine/Source/Editor/Kismet/Private/NativeCodeGenerationTool.cpp
--- a/Engine/Source/Editor/Kismet/Private/NativeCodeGenerationTool.cpp
+++ b/Engine/Source/Editor/Kismet/Private/NativeCodeGenerationTool.cpp
@@ -53,6 +53,11 @@ struct FGeneratedCodeData
 				}
 			}
 
+			if (InBlueprint.GeneratedClass)
+			{
+				ReferenceFinder.FindReferences(InBlueprint.GeneratedClass->GetDefaultObject(false));
+			}
+
 			for (UClass* Class = InBlueprint.GeneratedClass->GetSuperClass(); Class && !Class->HasAnyClassFlags(CLASS_Native); Class = Class->GetSuperClass())
 			{
 				ReferencedObjects.Add(Class);
diff --git a/Engine/Source/Runtime/Engine/Classes/Engine/InheritableComponentHandler.h b/Engine/Source/Runtime/Engine/Classes/Engine/InheritableComponentHandler.h
--- a/Engine/Source/Runtime/Engine/Classes/Engine/InheritableComponentHandler.h
+++ b/Engine/Source/Runtime/Engine/Classes/Engine/InheritableComponentHandler.h
@@ -8,6 +8,7 @@ class  USCS_Node;
 class  UActorComponent;
 struct FUCSComponentId;
 class  UBlueprint;
+class  UBlueprintGeneratedClass;
 
 USTRUCT()
 struct ENGINE_API FComponentKey
