Object properties are now forward declared, this allows circularly dependent blueprints to compile

diff --git a/Engine/Source/Developer/BlueprintCompilerCppBackend/Private/BlueprintCompilerCppBackendBase.cpp b/Engine/Source/Developer/BlueprintCompilerCppBackend/Private/BlueprintCompilerCppBackendBase.cpp
--- a/Engine/Source/Developer/BlueprintCompilerCppBackend/Private/BlueprintCompilerCppBackendBase.cpp
+++ b/Engine/Source/Developer/BlueprintCompilerCppBackend/Private/BlueprintCompilerCppBackendBase.cpp
@@ -384,6 +384,7 @@ void FBlueprintCompilerCppBackendBase::EmitFileBeginning(const FString& CleanNam
 	{ 
 		// find objects referenced by functions/script
 		TArray<UObject*> IncludeInHeader;
+		TArray<const UStruct*> DeclareInHeader;
 		TArray<UObject*> IncludeInBody;
 		{
 			FReferenceFinder HeaderReferenceFinder(IncludeInHeader, NULL, false, false, true, false);
@@ -399,7 +400,15 @@ void FBlueprintCompilerCppBackendBase::EmitFileBeginning(const FString& CleanNam
 					const bool bIsMemberVariable = OwnerProperty && (OwnerProperty->GetOuter() == SourceStruct);
 					if (bIsParam || bIsMemberVariable)
 					{
-						HeaderReferenceFinder.FindReferences(Obj);
+						const UObjectProperty* ObjectProperty = Cast<UObjectProperty>(OwnerProperty);
+						if (ObjectProperty)
+						{
+							DeclareInHeader.AddUnique(ObjectProperty->PropertyClass);
+						}
+						else
+						{
+							HeaderReferenceFinder.FindReferences(Obj);
+						}
 					}
 					else
 					{
@@ -437,7 +446,7 @@ void FBlueprintCompilerCppBackendBase::EmitFileBeginning(const FString& CleanNam
 
 		TSet<FString> AlreadyIncluded;
 		AlreadyIncluded.Add(SourceStruct->GetName());
-		auto EmitInner = [&](FStringOutputDevice& Dst, const TArray<UObject*>& Src)
+		auto EmitInner = [&](FStringOutputDevice& Dst, const TArray<UObject*>& Src, const TArray<const UStruct*>& Declarations )
 		{
 			auto EngineSourceDir = FPaths::EngineSourceDir();
 			auto GameSourceDir = FPaths::GameSourceDir();
@@ -501,11 +510,19 @@ void FBlueprintCompilerCppBackendBase::EmitFileBeginning(const FString& CleanNam
 
 				}
 			}
+
+			Emit(Dst, TEXT("\n"));
+
+			for (const UStruct* Type : Declarations)
+			{
+				Emit(Dst, *FString::Printf(TEXT("class %s;\n"), *(FString(Type->GetPrefixCPP()) + Type->GetName())));
+			}
+
 			Emit(Dst, TEXT("\n"));
 		};
 
-		EmitInner(Header, IncludeInHeader);
-		EmitInner(Body, IncludeInBody);
+		EmitInner(Header, IncludeInHeader, DeclareInHeader);
+		EmitInner(Body, IncludeInBody, TArray<const UStruct*>() );
 	}
 	Emit(Header, *FString::Printf(TEXT("#include \"%s.generated.h\"\n\n"), *CleanName));
 }
