UEBP-40 C++ code generation (WIP):
- support for UMG (WIP)

diff --git a/Engine/Source/Developer/BlueprintCompilerCppBackend/BlueprintCompilerCppBackend.Build.cs b/Engine/Source/Developer/BlueprintCompilerCppBackend/BlueprintCompilerCppBackend.Build.cs
--- a/Engine/Source/Developer/BlueprintCompilerCppBackend/BlueprintCompilerCppBackend.Build.cs
+++ b/Engine/Source/Developer/BlueprintCompilerCppBackend/BlueprintCompilerCppBackend.Build.cs
@@ -29,6 +29,13 @@ namespace UnrealBuildTool.Rules
                     "BlueprintGraph",
 				}
 				);
+            PrivateDependencyModuleNames.AddRange(
+                new string[]
+				{
+					"UMG",
+                    "SlateCore",
+				}
+                );
 		}
 	}
 }
\ No newline at end of file
diff --git a/Engine/Source/Developer/BlueprintCompilerCppBackend/Private/BlueprintCompilerCppBackendBase.cpp b/Engine/Source/Developer/BlueprintCompilerCppBackend/Private/BlueprintCompilerCppBackendBase.cpp
--- a/Engine/Source/Developer/BlueprintCompilerCppBackend/Private/BlueprintCompilerCppBackendBase.cpp
+++ b/Engine/Source/Developer/BlueprintCompilerCppBackend/Private/BlueprintCompilerCppBackendBase.cpp
@@ -138,6 +138,8 @@ void FBlueprintCompilerCppBackendBase::GenerateCodeFromClass(UClass* SourceClass
 		}
 	}
 
+	Emit(Header, *FBackendHelperUMG::WidgetFunctionsInHeader(SourceClass));
+
 	Emit(Header, TEXT("};\n\n"));
 
 	Emit(Body, *FEmitHelper::EmitLifetimeReplicatedPropsImpl(SourceClass, CppClassName, TEXT("")));
diff --git a/Engine/Source/Developer/BlueprintCompilerCppBackend/Private/BlueprintCompilerCppBackendUMG.cpp b/Engine/Source/Developer/BlueprintCompilerCppBackend/Private/BlueprintCompilerCppBackendUMG.cpp
new file mode 100644
--- /dev/null
+++ b/Engine/Source/Developer/BlueprintCompilerCppBackend/Private/BlueprintCompilerCppBackendUMG.cpp
@@ -0,0 +1,98 @@
+// Copyright 1998-2015 Epic Games, Inc. All Rights Reserved.
+
+#include "BlueprintCompilerCppBackendModulePrivatePCH.h"
+#include "BlueprintCompilerCppBackendUtils.h"
+#include "WidgetBlueprintGeneratedClass.h"
+#include "WidgetAnimation.h"
+#include "WidgetTree.h"
+
+FString FBackendHelperUMG::WidgetFunctionsInHeader(UClass* SourceClass)
+{
+	if (Cast<UWidgetBlueprintGeneratedClass>(SourceClass))
+	{
+		return FString(TEXT("public:\n\tvirtual void GetSlotNames(TArray<FName>& SlotNames) const override;\n\tvirtual void PreSave() override;\n\tvirtual void CustomNativeInitilize() override;\n"));
+	}
+	return FString();
+}
+
+void FBackendHelperUMG::CreateClassSubobjects(FEmitterLocalContext& Context)
+{
+	if (auto WidgetClass = Cast<UWidgetBlueprintGeneratedClass>(Context.GetCurrentlyGeneratedClass()))
+	{
+		if (WidgetClass->WidgetTree)
+		{
+			ensure(WidgetClass->WidgetTree->GetOuter() == Context.GetCurrentlyGeneratedClass());
+			FEmitDefaultValueHelper::HandleClassSubobject(Context, WidgetClass->WidgetTree);
+		}
+		for (auto Anim : WidgetClass->Animations)
+		{
+			ensure(Anim->GetOuter() == Context.GetCurrentlyGeneratedClass());
+			FEmitDefaultValueHelper::HandleClassSubobject(Context, WidgetClass->WidgetTree);
+		}
+	}
+}
+
+void FBackendHelperUMG::EmitWidgetInitializationFunctions(FEmitterLocalContext& Context)
+{
+	if (auto WidgetClass = Cast<UWidgetBlueprintGeneratedClass>(Context.GetCurrentlyGeneratedClass()))
+	{
+		const FString CppClassName = FString(WidgetClass->GetPrefixCPP()) + WidgetClass->GetName();
+
+		auto GenerateLocalProperty = [](FEmitterLocalContext& InContext, UProperty* InProperty, const uint8* DataPtr) -> FString
+		{
+			check(InProperty && DataPtr);
+			const FString NativeName = InContext.GenerateUniqueLocalName();
+			FStringOutputDevice Target;
+			InProperty->ExportCppDeclaration(Target, EExportedDeclaration::Local, nullptr
+				, EPropertyExportCPPFlags::CPPF_CustomTypeName | EPropertyExportCPPFlags::CPPF_BlueprintCppBackend | EPropertyExportCPPFlags::CPPF_NoConst | EPropertyExportCPPFlags::CPPF_NoRef
+				, true);
+			InContext.AddLine(FString::Printf(TEXT("%s %s;"), *Target, *NativeName));
+			FEmitDefaultValueHelper::InnerGenerate(InContext, InProperty, NativeName, DataPtr, nullptr, true);
+			return NativeName;
+		};
+
+		{	// GetSlotNames
+			Context.AddLine(FString::Printf(TEXT("void %s::GetSlotNames(TArray<FName>& SlotNames) const"), *CppClassName));
+			Context.AddLine(TEXT("{"));
+			Context.IncreaseIndent();
+
+			const FString LocalNativeName = GenerateLocalProperty(Context, FindFieldChecked<UArrayProperty>(UWidgetBlueprintGeneratedClass::StaticClass(), TEXT("NamedSlots")), reinterpret_cast<const uint8*>(&WidgetClass->NamedSlots));
+			Context.AddLine(FString::Printf(TEXT("SlotNames.Append(%s);"), *LocalNativeName));
+
+			Context.DecreaseIndent();
+			Context.AddLine(TEXT("}"));
+		}
+
+		{	// CustomNativeInitilize
+			Context.AddLine(FString::Printf(TEXT("void %s::CustomNativeInitilize()"), *CppClassName));
+			Context.AddLine(TEXT("{"));
+			Context.IncreaseIndent();
+
+			const FString WidgetTreeStr = Context.FindGloballyMappedObject(WidgetClass->WidgetTree, true);
+			ensure(!WidgetTreeStr.IsEmpty());
+			const FString AnimationsArrayNativeName = GenerateLocalProperty(Context, FindFieldChecked<UArrayProperty>(UWidgetBlueprintGeneratedClass::StaticClass(), TEXT("Animations")), reinterpret_cast<const uint8*>(&WidgetClass->Animations));
+			const FString BindingsArrayNativeName = GenerateLocalProperty(Context, FindFieldChecked<UArrayProperty>(UWidgetBlueprintGeneratedClass::StaticClass(), TEXT("Bindings")), reinterpret_cast<const uint8*>(&WidgetClass->Bindings));
+					
+			Context.AddLine(FString::Printf(TEXT("UWidgetBlueprintGeneratedClass::InitializeWidgetStatic(this, GetClass(), %s, %s, %s, %s, %s, nullptr);")
+				, WidgetClass->bCanEverTick ? TEXT("true") : TEXT("false")
+				, WidgetClass->bCanEverPaint ? TEXT("true") : TEXT("false")
+				, *WidgetTreeStr
+				, *AnimationsArrayNativeName
+				, *BindingsArrayNativeName));
+
+			Context.DecreaseIndent();
+			Context.AddLine(TEXT("}"));
+		}
+
+		// PreSave
+		Context.AddLine(FString::Printf(TEXT("void %s::PreSave()"), *CppClassName));
+		Context.AddLine(TEXT("{"));
+		Context.IncreaseIndent();
+		Context.AddLine(TEXT("Super::PreSave();"));
+		Context.AddLine(TEXT("TArray<FName> LocalNamedSlots;"));
+		Context.AddLine(TEXT("GetSlotNames(LocalNamedSlots);"));
+		Context.AddLine(TEXT("RemoveObsoleteBindings(LocalNamedSlots);"));
+		Context.DecreaseIndent();
+		Context.AddLine(TEXT("}"));
+	}
+}
diff --git a/Engine/Source/Developer/BlueprintCompilerCppBackend/Private/BlueprintCompilerCppBackendUtils.h b/Engine/Source/Developer/BlueprintCompilerCppBackend/Private/BlueprintCompilerCppBackendUtils.h
--- a/Engine/Source/Developer/BlueprintCompilerCppBackend/Private/BlueprintCompilerCppBackendUtils.h
+++ b/Engine/Source/Developer/BlueprintCompilerCppBackend/Private/BlueprintCompilerCppBackendUtils.h
@@ -52,13 +52,14 @@ public:
 		NativeObjectNamesInConstructor.Add(Object, NativeName);
 	}
 
-	FString AddNewObject_InConstructor(UObject* Object)
+	FString AddNewObject_InConstructor(UObject* Object, bool bAddToObjectsCreatedPerClass)
 	{
 		ensure(bInsideConstructor);
 		const FString UniqueName = GenerateUniqueLocalName();
 		NativeObjectNamesInConstructor.Add(Object, UniqueName);
-		if (bCreatingObjectsPerClass)
+		if (bAddToObjectsCreatedPerClass)
 		{
+			ensure(bCreatingObjectsPerClass);
 			check(!ObjectsCreatedPerClass.Contains(Object));
 			ObjectsCreatedPerClass.Add(Object);
 		}
@@ -104,7 +105,7 @@ public:
 
 	/** All objects (that can be referenced from other package) that will have a different path in cooked build 
 	(due to the native code generation), should be handled by this function */
-	FString FindGloballyMappedObject(UObject* Object)
+	FString FindGloballyMappedObject(UObject* Object, bool bLoadIfNotFound = false)
 	{
 		// TODO: check if not excluded
 
@@ -144,8 +145,26 @@ public:
 			return FString::Printf(TEXT("FindObjectChecked<UEnum>(ANY_PACKAGE, TEXT(\"%s\"))"), *UDE->GetName());
 		}
 
+		int32 ObjectsCreatedPerClassIdx = INDEX_NONE;
+		if (Object && ObjectsCreatedPerClass.Find(Object, ObjectsCreatedPerClassIdx))
+		{
+			return FString::Printf(TEXT("CastChecked<%s%s>(%s%s::StaticClass()->ConvertedSubobjectsFromBPGC[%d])")
+				, Object->GetClass()->GetPrefixCPP()
+				, *Object->GetClass()->GetName()
+				, ActualClass->GetPrefixCPP()
+				, *ActualClass->GetName()
+				, ObjectsCreatedPerClassIdx);
+		}
+
 		// TODO: handle subobjects
 
+		if (bLoadIfNotFound && ensure(Object))
+		{
+			UClass* FoundClass = Object->GetClass();
+			const FString ClassString = FString(FoundClass->GetPrefixCPP()) + FoundClass->GetName();
+			return FString::Printf(TEXT("LoadObject<%s>(nullptr, TEXT(\"%s\"))"), *ClassString, *(Object->GetPathName().ReplaceCharWithEscapedChar()));
+		}
+
 		return FString{};
 	}
 
@@ -259,19 +278,30 @@ struct FEmitDefaultValueHelper
 	// OuterPath ends context/outer name (or empty, if the scope is "this")
 	static void OuterGenerate(FEmitterLocalContext& Context, const UProperty* Property, const FString& OuterPath, const uint8* DataContainer, const uint8* OptionalDefaultDataContainer, EPropertyAccessOperator AccessOperator, bool bAllowProtected = false);
 
-private:
+
 	// PathToMember ends with variable name
 	static void InnerGenerate(FEmitterLocalContext& Context, const UProperty* Property, const FString& PathToMember, const uint8* ValuePtr, const uint8* DefaultValuePtr, bool bWithoutFirstConstructionLine = false);
 
+	// Creates the subobject (of class) returns it's native local name, 
+	// returns empty string if cannot handle
+	static FString HandleClassSubobject(FEmitterLocalContext& Context, UObject* Object);
+
+private:
 	// Returns native term, 
 	// returns empty string if cannot handle
 	static FString HandleSpecialTypes(FEmitterLocalContext& Context, const UProperty* Property, const uint8* ValuePtr);
 
 	static FString HandleNonNativeComponent(FEmitterLocalContext& Context, const USCS_Node* Node, TSet<const UProperty*>& OutHandledProperties, const USCS_Node* ParentNode = nullptr);
-	
-	// Creates the subobject (of class) returns it's native local name, 
-	// returns empty string if cannot handle
-	static FString HandleClassSubobject(FEmitterLocalContext& Context, UObject* Object);
 
 	static FString HandleInstancedSubobject(FEmitterLocalContext& Context, UObject* Object, bool bSkipEditorOnlyCheck = false);
+};
+
+struct FBackendHelperUMG
+{
+	static FString WidgetFunctionsInHeader(UClass* SourceClass);
+
+	// these function should use the same context as Constructor
+	static void CreateClassSubobjects(FEmitterLocalContext& Context);
+	static void EmitWidgetInitializationFunctions(FEmitterLocalContext& Context);
+
 };
\ No newline at end of file
diff --git a/Engine/Source/Developer/BlueprintCompilerCppBackend/Private/BlueprintCompilerCppBackendValueHelper.cpp b/Engine/Source/Developer/BlueprintCompilerCppBackend/Private/BlueprintCompilerCppBackendValueHelper.cpp
--- a/Engine/Source/Developer/BlueprintCompilerCppBackend/Private/BlueprintCompilerCppBackendValueHelper.cpp
+++ b/Engine/Source/Developer/BlueprintCompilerCppBackend/Private/BlueprintCompilerCppBackendValueHelper.cpp
@@ -385,6 +385,8 @@ FString FEmitDefaultValueHelper::GenerateConstructor(UClass* InBPGC)
 			}
 		}
 
+		FBackendHelperUMG::CreateClassSubobjects(Context);
+
 		Context.bCreatingObjectsPerClass = false;
 		Context.DecreaseIndent();
 		Context.AddLine(TEXT("}"));
@@ -502,6 +504,9 @@ FString FEmitDefaultValueHelper::GenerateConstructor(UClass* InBPGC)
 	Context.DecreaseIndent();
 	Context.AddLine(TEXT("}"));
 	Context.bInsideConstructor = false;
+
+	FBackendHelperUMG::EmitWidgetInitializationFunctions(Context);
+
 	return Context.GetResult();
 }
 
@@ -513,7 +518,8 @@ FString FEmitDefaultValueHelper::HandleClassSubobject(FEmitterLocalContext& Cont
 		return FString();
 	}
 
-	const FString LocalNativeName = Context.AddNewObject_InConstructor(Object);
+	const bool AddAsSubobjectOfClass = Object->GetOuter() == Context.GetCurrentlyGeneratedClass();
+	const FString LocalNativeName = Context.AddNewObject_InConstructor(Object, AddAsSubobjectOfClass);
 	UClass* ObjectClass = Object->GetClass();
 	Context.AddHighPriorityLine(FString::Printf(
 		TEXT("auto %s = NewObject<%s%s>(%s, TEXT(\"%s\"));")
@@ -521,9 +527,10 @@ FString FEmitDefaultValueHelper::HandleClassSubobject(FEmitterLocalContext& Cont
 		, ObjectClass->GetPrefixCPP()
 		, *ObjectClass->GetName()
 		, *OuterStr, *Object->GetName()));
-	if (Object->GetOuter() == Context.GetCurrentlyGeneratedClass())
+	if (AddAsSubobjectOfClass)
 	{
-		Context.AddLine(FString::Printf(TEXT("GetClass()->ConvertedSubobjectsFromBPGC.Add(%s);"), *LocalNativeName));
+		ensure(Context.bCreatingObjectsPerClass);
+		Context.AddHighPriorityLine(FString::Printf(TEXT("GetClass()->ConvertedSubobjectsFromBPGC.Add(%s);"), *LocalNativeName));
 	}
 
 	for (auto Property : TFieldRange<const UProperty>(ObjectClass))
@@ -549,7 +556,7 @@ FString FEmitDefaultValueHelper::HandleInstancedSubobject(FEmitterLocalContext&
 	}
 	else
 	{
-		LocalNativeName = Context.AddNewObject_InConstructor(Object);
+		LocalNativeName = Context.AddNewObject_InConstructor(Object, false);
 	}
 
 	UClass* ObjectClass = Object->GetClass();
diff --git a/Engine/Source/Developer/BlueprintNativeCodeGen/Private/NativeCodeGenerationTool.cpp b/Engine/Source/Developer/BlueprintNativeCodeGen/Private/NativeCodeGenerationTool.cpp
--- a/Engine/Source/Developer/BlueprintNativeCodeGen/Private/NativeCodeGenerationTool.cpp
+++ b/Engine/Source/Developer/BlueprintNativeCodeGen/Private/NativeCodeGenerationTool.cpp
@@ -469,7 +469,7 @@ bool FNativeCodeGenerationTool::CanGenerate(const UBlueprint& Blueprint)
 {
 	return (Blueprint.Status == EBlueprintStatus::BS_UpToDate)
 		&& (Blueprint.BlueprintType == EBlueprintType::BPTYPE_Normal || Blueprint.BlueprintType == EBlueprintType::BPTYPE_FunctionLibrary)
-		&& (Blueprint.GeneratedClass->GetClass() == UBlueprintGeneratedClass::StaticClass());
+		&& Cast<UBlueprintGeneratedClass>(Blueprint.GeneratedClass);
 }
 
 #undef LOCTEXT_NAMESPACE
\ No newline at end of file
