UEBP-40 C++ code generation (WIP):
- Support for dynamically loaded native uclass object
- for classes that are forward-declared in header, there is a proper include in .cpp
- delegate declaration are moved into class scope

diff --git a/Engine/Source/Developer/BlueprintCompilerCppBackend/Private/BlueprintCompilerCppBackendBase.cpp b/Engine/Source/Developer/BlueprintCompilerCppBackend/Private/BlueprintCompilerCppBackendBase.cpp
--- a/Engine/Source/Developer/BlueprintCompilerCppBackend/Private/BlueprintCompilerCppBackendBase.cpp
+++ b/Engine/Source/Developer/BlueprintCompilerCppBackend/Private/BlueprintCompilerCppBackendBase.cpp
@@ -30,13 +30,8 @@ void FBlueprintCompilerCppBackendBase::EmitStructProperties(FStringOutputDevice&
 	}
 }
 
-void FBlueprintCompilerCppBackendBase::GenerateCodeFromClass(UClass* SourceClass, TIndirectArray<FKismetFunctionContext>& Functions, bool bGenerateStubsOnly)
+void FBlueprintCompilerCppBackendBase::DeclareDelegates(UClass* SourceClass, TIndirectArray<FKismetFunctionContext>& Functions)
 {
-	auto CleanCppClassName = SourceClass->GetName();
-	CppClassName = FString(SourceClass->GetPrefixCPP()) + CleanCppClassName;
-	
-	EmitFileBeginning(CleanCppClassName, SourceClass);
-
 	// MC DELEGATE DECLARATION
 	{
 		auto DelegateDeclarations = FEmitHelper::EmitMulticastDelegateDeclarations(SourceClass);
@@ -74,6 +69,14 @@ void FBlueprintCompilerCppBackendBase::GenerateCodeFromClass(UClass* SourceClass
 			Emit(Header, TEXT(";\n"));
 		}
 	}
+}
+
+void FBlueprintCompilerCppBackendBase::GenerateCodeFromClass(UClass* SourceClass, TIndirectArray<FKismetFunctionContext>& Functions, bool bGenerateStubsOnly)
+{
+	auto CleanCppClassName = SourceClass->GetName();
+	CppClassName = FString(SourceClass->GetPrefixCPP()) + CleanCppClassName;
+	
+	EmitFileBeginning(CleanCppClassName, SourceClass);
 
 	// Class declaration
 	const bool bIsInterface = SourceClass->IsChildOf<UInterface>();
@@ -109,6 +112,8 @@ void FBlueprintCompilerCppBackendBase::GenerateCodeFromClass(UClass* SourceClass
 	// Begin scope
 	Emit(Header,TEXT("\n{\npublic:\n\tGENERATED_BODY()\n"));
 
+	DeclareDelegates(SourceClass, Functions);
+
 	EmitStructProperties(Header, SourceClass);
 
 	// Create the state map
@@ -138,6 +143,9 @@ void FBlueprintCompilerCppBackendBase::GenerateCodeFromClass(UClass* SourceClass
 		}
 	}
 
+	Emit(Header, TEXT("\n\tstatic void __StaticDependenciesConvertedClasses(TArray<FName>& OutNames);\n"));
+	Emit(Header, TEXT("\n\tstatic void __StaticDependenciesAssets(TArray<FName>& OutNames);\n"));
+
 	Emit(Header, *FBackendHelperUMG::WidgetFunctionsInHeader(SourceClass));
 
 	Emit(Header, TEXT("};\n\n"));
@@ -406,6 +414,7 @@ void FBlueprintCompilerCppBackendBase::EmitFileBeginning(const FString& CleanNam
 						if (ObjectProperty)
 						{
 							DeclareInHeader.AddUnique(ObjectProperty->PropertyClass);
+							BodyReferenceFinder.FindReferences(Obj);
 						}
 						else
 						{
diff --git a/Engine/Source/Developer/BlueprintCompilerCppBackend/Private/BlueprintCompilerCppBackendBase.h b/Engine/Source/Developer/BlueprintCompilerCppBackend/Private/BlueprintCompilerCppBackendBase.h
--- a/Engine/Source/Developer/BlueprintCompilerCppBackend/Private/BlueprintCompilerCppBackendBase.h
+++ b/Engine/Source/Developer/BlueprintCompilerCppBackend/Private/BlueprintCompilerCppBackendBase.h
@@ -88,4 +88,6 @@ protected:
 	}
 
 	void EmitReplaceConvertedMetaData(UObject* Obj);
+
+	void DeclareDelegates(UClass* SourceClass, TIndirectArray<FKismetFunctionContext>& Functions);
 };
diff --git a/Engine/Source/Developer/BlueprintCompilerCppBackend/Private/BlueprintCompilerCppBackendUtils.cpp b/Engine/Source/Developer/BlueprintCompilerCppBackend/Private/BlueprintCompilerCppBackendUtils.cpp
--- a/Engine/Source/Developer/BlueprintCompilerCppBackend/Private/BlueprintCompilerCppBackendUtils.cpp
+++ b/Engine/Source/Developer/BlueprintCompilerCppBackend/Private/BlueprintCompilerCppBackendUtils.cpp
@@ -331,7 +331,7 @@ TArray<FString> FEmitHelper::EmitMulticastDelegateDeclarations(UClass* SourceCla
 		FString ParamNumberStr, Parameters;
 		ParseDelegateDetails(Signature, Parameters, ParamNumberStr);
 
-		Results.Add(*FString::Printf(TEXT("DECLARE_DYNAMIC_MULTICAST_DELEGATE%s(%s%s)"),
+		Results.Add(*FString::Printf(TEXT("\tDECLARE_DYNAMIC_MULTICAST_DELEGATE%s(%s%s)"),
 			*ParamNumberStr, *It->GetCPPType(NULL, EPropertyExportCPPFlags::CPPF_CustomTypeName), *Parameters));
 	}
 
diff --git a/Engine/Source/Developer/BlueprintCompilerCppBackend/Private/BlueprintCompilerCppBackendUtils.h b/Engine/Source/Developer/BlueprintCompilerCppBackend/Private/BlueprintCompilerCppBackendUtils.h
--- a/Engine/Source/Developer/BlueprintCompilerCppBackend/Private/BlueprintCompilerCppBackendUtils.h
+++ b/Engine/Source/Developer/BlueprintCompilerCppBackend/Private/BlueprintCompilerCppBackendUtils.h
@@ -103,6 +103,14 @@ public:
 		return ActualClass;
 	}
 
+	bool WillClassBeConverted(UClass* InClass) const
+	{
+		// TODO:
+
+		auto BPGC = Cast<UBlueprintGeneratedClass>(InClass);
+		return (nullptr != BPGC) && !InClass->HasAnyFlags(RF_ClassDefaultObject);
+	}
+
 	/** All objects (that can be referenced from other package) that will have a different path in cooked build 
 	(due to the native code generation), should be handled by this function */
 	FString FindGloballyMappedObject(UObject* Object, bool bLoadIfNotFound = false)
@@ -116,7 +124,7 @@ public:
 
 		if (auto ObjClass = Cast<UClass>(Object))
 		{
-			if (ObjClass->HasAnyClassFlags(CLASS_Native))
+			if (ObjClass->HasAnyClassFlags(CLASS_Native) || WillClassBeConverted(ObjClass))
 			{
 				return FString::Printf(TEXT("%s%s::StaticClass()"), ObjClass->GetPrefixCPP(), *ObjClass->GetName());
 			}
@@ -124,12 +132,6 @@ public:
 
 		// TODO Handle native structires, and special cases..
 
-		if (auto BPGC = ExactCast<UBlueprintGeneratedClass>(Object))
-		{
-			// TODO: check if supported 
-			return FString::Printf(TEXT("%s%s::StaticClass()"), BPGC->GetPrefixCPP(), *BPGC->GetName());
-		}
-
 		if (auto UDS = Cast<UScriptStruct>(Object))
 		{
 			// Check if  
diff --git a/Engine/Source/Developer/BlueprintCompilerCppBackend/Private/BlueprintCompilerCppBackendValueHelper.cpp b/Engine/Source/Developer/BlueprintCompilerCppBackend/Private/BlueprintCompilerCppBackendValueHelper.cpp
--- a/Engine/Source/Developer/BlueprintCompilerCppBackend/Private/BlueprintCompilerCppBackendValueHelper.cpp
+++ b/Engine/Source/Developer/BlueprintCompilerCppBackend/Private/BlueprintCompilerCppBackendValueHelper.cpp
@@ -323,6 +323,150 @@ FString FEmitDefaultValueHelper::HandleNonNativeComponent(FEmitterLocalContext&
 	return VariableName;
 }
 
+struct FGatherDependencies : public FReferenceCollector
+{
+	TSet<UObject*> SerializedObjects;
+	TSet<UObject*> Assets;
+	TSet<UStruct*> ConvertedClasses;
+
+	FEmitterLocalContext& Context;
+	UClass* OriginalClass;
+
+	void FindReferences(UObject* Object)
+	{
+		{
+			FSimpleObjectReferenceCollectorArchive CollectorArchive(Object, *this);
+			CollectorArchive.SetSerializedProperty(nullptr);
+			CollectorArchive.SetFilterEditorOnly(true);
+			Object->SerializeScriptProperties(CollectorArchive);
+		}
+		Object->CallAddReferencedObjects(*this);
+	}
+
+	virtual bool IsIgnoringArchetypeRef() const override { return false; }
+	virtual bool IsIgnoringTransient() const override { return true; }
+
+	virtual void HandleObjectReference(UObject*& InObject, const UObject* InReferencingObject, const UProperty* InReferencingProperty) override
+	{
+		if (InObject)
+		{
+			if (InObject->IsA<UBlueprint>() || (InObject == OriginalClass))
+			{
+				return;
+			}
+
+			// converted class
+			//TODO: What About Delegates?
+			auto ObjAsBPGC = Cast<UBlueprintGeneratedClass>(InObject);
+			const bool SaveAsConvertedClass = ObjAsBPGC && Context.WillClassBeConverted(ObjAsBPGC);
+			auto ObjAsUDS = Cast<UUserDefinedStruct>(InObject);// all UDS must be converted
+			const bool SaveAsConvertedStruct = ObjAsUDS && !ObjAsUDS->HasAnyFlags(RF_ClassDefaultObject);
+			if (SaveAsConvertedClass || SaveAsConvertedStruct)
+			{
+				ConvertedClasses.Add(CastChecked<UStruct>(InObject));
+			}
+			// asset
+			else if (InObject->IsAsset() 
+				&& !InObject->IsIn(OriginalClass) 
+				&& (InObject != OriginalClass)
+				&& !InObject->IsA<UUserDefinedEnum>()) // all UDE must be converted
+			{
+				Assets.Add(InObject);
+				return;
+			}
+
+			bool AlreadyAdded = false;
+			SerializedObjects.Add(InObject, &AlreadyAdded);
+			if (!AlreadyAdded)
+			{
+				FindReferences(InObject);
+			}
+		}
+	}
+
+	FGatherDependencies(UClass* InClass, FEmitterLocalContext& InContext) 
+		: OriginalClass(InClass)
+		, Context(InContext)
+	{
+		check(OriginalClass);
+		FindReferences(OriginalClass);
+	}
+
+	void AddDependenciesInConstructor()
+	{
+		Context.AddLine(TEXT("// List of all referenced converted classes"));
+		for (auto LocStruct : ConvertedClasses)
+		{
+			const bool bIsClass = nullptr != Cast<UClass>(LocStruct);
+			Context.AddLine(FString::Printf(TEXT("GetClass()->ConvertedSubobjectsFromBPGC.Add(%s%s::%s());")
+				, LocStruct->GetPrefixCPP()
+				, *LocStruct->GetName()
+				, (bIsClass ? TEXT("StaticClass") : TEXT("StaticStruct"))));
+		}
+		Context.AddLine(TEXT("// List of all referenced assets"));
+		for (auto LocAsset : Assets)
+		{
+			const FString AssetStr = Context.FindGloballyMappedObject(LocAsset, true);
+			Context.AddLine(FString::Printf(TEXT("GetClass()->ConvertedSubobjectsFromBPGC.Add(%s);"), *AssetStr));
+		}
+	}
+
+	void AddStaticFunctionsForDependencies()
+	{
+		const FString CppClassName = FString::Printf(TEXT("%s%s"), OriginalClass->GetPrefixCPP(), *OriginalClass->GetName());
+		// __StaticDependenciesConvertedClasses
+		Context.AddLine(FString::Printf(TEXT("void %s::__StaticDependenciesConvertedClasses(TArray<FName>& OutNames)"), *CppClassName));
+		Context.AddLine(TEXT("{"));
+		Context.IncreaseIndent();
+
+		for (auto LocClass : ConvertedClasses)
+		{
+			Context.AddLine(FString::Printf(TEXT("OutNames.Add(TEXT(\"%s\"));"), *LocClass->GetName()));
+		}
+
+		Context.DecreaseIndent();
+		Context.AddLine(TEXT("}"));
+
+		// __StaticDependenciesAssets
+		Context.AddLine(FString::Printf(TEXT("void %s::__StaticDependenciesAssets(TArray<FName>& OutPackagePaths)"), *CppClassName));
+		Context.AddLine(TEXT("{"));
+		Context.IncreaseIndent();
+
+		for (UObject* LocAsset : Assets)
+		{
+			const FString PakagePath = LocAsset->GetOutermost()->GetPathName();
+			Context.AddLine(FString::Printf(TEXT("OutPackagePaths.Add(TEXT(\"%s\"));"), *PakagePath));
+		}
+
+		Context.DecreaseIndent();
+		Context.AddLine(TEXT("}"));
+
+		// Register Helper Struct
+		const FString RegisterHelperName = FString::Printf(TEXT("FRegisterHelper__%s"), *CppClassName);
+		Context.AddLine(FString::Printf(TEXT("struct %s"), *RegisterHelperName));
+		Context.AddLine(TEXT("{"));
+		Context.IncreaseIndent();
+
+		Context.AddLine(FString::Printf(TEXT("%s()"), *RegisterHelperName));
+		Context.AddLine(TEXT("{"));
+		Context.IncreaseIndent();
+		Context.AddLine(FString::Printf(
+			TEXT("FConvertedBlueprintsDependencies::Get().RegisterClass(TEXT(\"%s\"), &%s::__StaticDependenciesConvertedClasses, &%s::__StaticDependenciesAssets);")
+			, *OriginalClass->GetName()
+			, *CppClassName
+			, *CppClassName));
+		Context.DecreaseIndent();
+		Context.AddLine(TEXT("}"));
+
+		Context.AddLine(FString::Printf(TEXT("static %s Instance;"), *RegisterHelperName));
+
+		Context.DecreaseIndent();
+		Context.AddLine(TEXT("};"));
+
+		Context.AddLine(FString::Printf(TEXT("%s %s::Instance;"), *RegisterHelperName, *RegisterHelperName));
+	}
+};
+
 FString FEmitDefaultValueHelper::GenerateConstructor(UClass* InBPGC)
 {
 	auto BPGC = CastChecked<UBlueprintGeneratedClass>(InBPGC);
@@ -334,6 +478,9 @@ FString FEmitDefaultValueHelper::GenerateConstructor(UClass* InBPGC)
 	Context.AddLine(TEXT("{"));
 	Context.IncreaseIndent();
 	Context.bInsideConstructor = true;
+
+	FGatherDependencies GatherDependencies(InBPGC, Context);
+
 	// When CDO is created create all subobjects owned by the class
 	TArray<UActorComponent*> ActorComponentTempatesOwnedByClass = BPGC->ComponentTemplates;
 	{
@@ -388,6 +535,9 @@ FString FEmitDefaultValueHelper::GenerateConstructor(UClass* InBPGC)
 		FBackendHelperUMG::CreateClassSubobjects(Context);
 
 		Context.bCreatingObjectsPerClass = false;
+
+		GatherDependencies.AddDependenciesInConstructor();
+
 		Context.DecreaseIndent();
 		Context.AddLine(TEXT("}"));
 
@@ -505,6 +655,8 @@ FString FEmitDefaultValueHelper::GenerateConstructor(UClass* InBPGC)
 	Context.AddLine(TEXT("}"));
 	Context.bInsideConstructor = false;
 
+	GatherDependencies.AddStaticFunctionsForDependencies();
+
 	FBackendHelperUMG::EmitWidgetInitializationFunctions(Context);
 
 	return Context.GetResult();
diff --git a/Engine/Source/Runtime/CoreUObject/Private/Blueprint/BlueprintSupport.cpp b/Engine/Source/Runtime/CoreUObject/Private/Blueprint/BlueprintSupport.cpp
--- a/Engine/Source/Runtime/CoreUObject/Private/Blueprint/BlueprintSupport.cpp
+++ b/Engine/Source/Runtime/CoreUObject/Private/Blueprint/BlueprintSupport.cpp
@@ -1722,4 +1722,42 @@ void FReplaceConvertedAssetManager::GatherOriginalPathsOfConvertedAssets()
 	}
 }
 
-#endif //WITH_EDITOR
\ No newline at end of file
+#endif //WITH_EDITOR
+
+FConvertedBlueprintsDependencies& FConvertedBlueprintsDependencies::Get()
+{
+	static FConvertedBlueprintsDependencies ConvertedBlueprintsDependencies;
+	return ConvertedBlueprintsDependencies;
+}
+
+void FConvertedBlueprintsDependencies::RegisterClass(FName ClassName, GetDependenciesNamesFunc GetConvertedClasses, GetDependenciesNamesFunc GetAssets)
+{
+	check(!ClassNameToGetter.Contains(ClassName));
+
+	FGetters Getters;
+	Getters.GetAssets = GetAssets;
+	Getters.GetConvertedClasses = GetConvertedClasses;
+	ClassNameToGetter.Add(ClassName, Getters);
+}
+
+void FConvertedBlueprintsDependencies::GetConvertedClasses(FName ClassName, TArray<FName>& OutConvertedClassNames) const
+{
+	auto GettersPtr = ClassNameToGetter.Find(ClassName);
+	auto Func = (GettersPtr) ? GettersPtr->GetConvertedClasses : nullptr;
+	ensure(Func || !GettersPtr);
+	if (Func)
+	{
+		Func(OutConvertedClassNames);
+	}
+}
+
+void FConvertedBlueprintsDependencies::GetAssets(FName ClassName, TArray<FName>& OutPackagePaths) const
+{
+	auto GettersPtr = ClassNameToGetter.Find(ClassName);
+	auto Func = (GettersPtr) ? GettersPtr->GetAssets : nullptr;
+	ensure(Func || !GettersPtr);
+	if (Func)
+	{
+		Func(OutPackagePaths);
+	}
+}
\ No newline at end of file
diff --git a/Engine/Source/Runtime/CoreUObject/Public/Blueprint/BlueprintSupport.h b/Engine/Source/Runtime/CoreUObject/Public/Blueprint/BlueprintSupport.h
--- a/Engine/Source/Runtime/CoreUObject/Public/Blueprint/BlueprintSupport.h
+++ b/Engine/Source/Runtime/CoreUObject/Public/Blueprint/BlueprintSupport.h
@@ -153,4 +153,45 @@ private:
 	UClass* ResolvingClass;
 	/** Tracks sub-classes that have had their CDO deferred as a result of the super not being fully serialized */
 	TMultiMap<UClass*, UClass*> SuperClassMap;
+};
+
+
+/**
+ *	Stores info about dependencies of native classes converted from BPs
+ */
+struct COREUOBJECT_API FConvertedBlueprintsDependencies
+{
+	typedef void(*GetDependenciesNamesFunc)(TArray<FName>&);
+
+private:
+
+	struct FGetters
+	{
+		GetDependenciesNamesFunc GetConvertedClasses;
+		GetDependenciesNamesFunc GetAssets;
+
+		FGetters()
+			: GetConvertedClasses(nullptr)
+			, GetAssets(nullptr)
+		{}
+	};
+
+	TMap<FName, FGetters> ClassNameToGetter;
+
+public:
+
+	static FConvertedBlueprintsDependencies& Get();
+
+	bool AnyClassRegistered() const
+	{
+		return ClassNameToGetter.Num() > 0;
+	}
+
+	void RegisterClass(FName ClassName, GetDependenciesNamesFunc GetConvertedClasses, GetDependenciesNamesFunc GetAssets);
+
+	/** Get all converted classes (and structures) names necessary for the class with the given class name. */
+	void GetConvertedClasses(FName ClassName, TArray<FName>& OutConvertedClassNames) const;
+
+	/** Get all assets paths necessary for the class with the given class name and all converted classes that dependencies. */
+	void GetAssets(FName ClassName, TArray<FName>& OutPackagePaths) const;
 };
\ No newline at end of file
diff --git a/Engine/Source/Runtime/Engine/Public/GeneratedCodeHelpers.h b/Engine/Source/Runtime/Engine/Public/GeneratedCodeHelpers.h
--- a/Engine/Source/Runtime/Engine/Public/GeneratedCodeHelpers.h
+++ b/Engine/Source/Runtime/Engine/Public/GeneratedCodeHelpers.h
@@ -4,6 +4,7 @@
 
 // Common includes
 #include "UObject/Stack.h"
+#include "Blueprint/BlueprintSupport.h"
 #include "Engine/BlueprintGeneratedClass.h"
 
 // Common libraries
