Fix initialization order

diff --git a/Engine/Source/Developer/BlueprintCompilerCppBackend/Private/BlueprintCompilerCppBackendValueHelper.cpp b/Engine/Source/Developer/BlueprintCompilerCppBackend/Private/BlueprintCompilerCppBackendValueHelper.cpp
--- a/Engine/Source/Developer/BlueprintCompilerCppBackend/Private/BlueprintCompilerCppBackendValueHelper.cpp
+++ b/Engine/Source/Developer/BlueprintCompilerCppBackend/Private/BlueprintCompilerCppBackendValueHelper.cpp
@@ -385,8 +385,8 @@ struct FGatherDependencies : public FReferenceCollector
 	}
 
 	FGatherDependencies(UClass* InClass, FEmitterLocalContext& InContext) 
-		: OriginalClass(InClass)
-		, Context(InContext)
+		: Context(InContext)
+		, OriginalClass(InClass)
 	{
 		check(OriginalClass);
 		FindReferences(OriginalClass);
