UEBP-40 C++ code generation (WIP):
- Gathering coverted class dependencies - fixed and refactored
- improved list of included headers and forward-declared structs

diff --git a/Engine/Source/Developer/BlueprintCompilerCppBackend/Private/BlueprintCompilerCppBackend.cpp b/Engine/Source/Developer/BlueprintCompilerCppBackend/Private/BlueprintCompilerCppBackend.cpp
--- a/Engine/Source/Developer/BlueprintCompilerCppBackend/Private/BlueprintCompilerCppBackend.cpp
+++ b/Engine/Source/Developer/BlueprintCompilerCppBackend/Private/BlueprintCompilerCppBackend.cpp
@@ -576,9 +576,9 @@ FString FBlueprintCompilerCppBackend::LatentFunctionInfoTermToText(FEmitterLocal
 	return FEmitHelper::LiteralTerm(EmitterContext, Term->Type, StructValues, nullptr);
 }
 
-FString FBlueprintCompilerCppBackend::InnerFunctionImplementation(FKismetFunctionContext& FunctionContext, bool bUseSwitchState)
+FString FBlueprintCompilerCppBackend::InnerFunctionImplementation(FKismetFunctionContext& FunctionContext, const FGatherConvertedClassDependencies& Dependencies, bool bUseSwitchState)
 {
-	FEmitterLocalContext EmitterContext;
+	FEmitterLocalContext EmitterContext(FunctionContext.NewClass, Dependencies);
 	EmitterContext.IncreaseIndent();
 	if (bUseSwitchState)
 	{
@@ -598,7 +598,6 @@ FString FBlueprintCompilerCppBackend::InnerFunctionImplementation(FKismetFunctio
 		EmitterContext.AddLine(TEXT("{"));
 		EmitterContext.IncreaseIndent();
 	}
-	EmitterContext.SetCurrentlyGeneratedClass(FunctionContext.NewClass);
 
 	// Emit code in the order specified by the linear execution list (the first node is always the entry point for the function)
 	for (int32 NodeIndex = 0; NodeIndex < FunctionContext.LinearExecutionList.Num(); ++NodeIndex)
diff --git a/Engine/Source/Developer/BlueprintCompilerCppBackend/Private/BlueprintCompilerCppBackend.h b/Engine/Source/Developer/BlueprintCompilerCppBackend/Private/BlueprintCompilerCppBackend.h
--- a/Engine/Source/Developer/BlueprintCompilerCppBackend/Private/BlueprintCompilerCppBackend.h
+++ b/Engine/Source/Developer/BlueprintCompilerCppBackend/Private/BlueprintCompilerCppBackend.h
@@ -9,7 +9,7 @@ struct FEmitterLocalContext;
 class FBlueprintCompilerCppBackend : public FBlueprintCompilerCppBackendBase
 {
 protected:
-	virtual FString InnerFunctionImplementation(FKismetFunctionContext& FunctionContext, bool bUseSwitchState) override;
+	virtual FString InnerFunctionImplementation(FKismetFunctionContext& FunctionContext, const FGatherConvertedClassDependencies& Dependencies, bool bUseSwitchState) override;
 
 protected:
 	void EmitCallStatment(FEmitterLocalContext& EmitterContext, FKismetFunctionContext& FunctionContext, FBlueprintCompiledStatement& Statement);
diff --git a/Engine/Source/Developer/BlueprintCompilerCppBackend/Private/BlueprintCompilerCppBackendBase.cpp b/Engine/Source/Developer/BlueprintCompilerCppBackend/Private/BlueprintCompilerCppBackendBase.cpp
--- a/Engine/Source/Developer/BlueprintCompilerCppBackend/Private/BlueprintCompilerCppBackendBase.cpp
+++ b/Engine/Source/Developer/BlueprintCompilerCppBackend/Private/BlueprintCompilerCppBackendBase.cpp
@@ -76,7 +76,8 @@ void FBlueprintCompilerCppBackendBase::GenerateCodeFromClass(UClass* SourceClass
 	auto CleanCppClassName = SourceClass->GetName();
 	CppClassName = FString(SourceClass->GetPrefixCPP()) + CleanCppClassName;
 	
-	EmitFileBeginning(CleanCppClassName, SourceClass);
+	FGatherConvertedClassDependencies Dependencies(SourceClass);
+	EmitFileBeginning(CleanCppClassName, &Dependencies);
 
 	// Class declaration
 	const bool bIsInterface = SourceClass->IsChildOf<UInterface>();
@@ -132,14 +133,14 @@ void FBlueprintCompilerCppBackendBase::GenerateCodeFromClass(UClass* SourceClass
 	if (!bIsInterface)
 	{
 		Emit(Header, *FString::Printf(TEXT("\t%s(const FObjectInitializer& ObjectInitializer = FObjectInitializer::Get());\n\n"), *CppClassName));
-		Emit(Body, *FEmitDefaultValueHelper::GenerateConstructor(SourceClass));
+		Emit(Body, *FEmitDefaultValueHelper::GenerateConstructor(SourceClass, Dependencies));
 	}
 
 	for (int32 i = 0; i < Functions.Num(); ++i)
 	{
 		if (Functions[i].IsValid())
 		{
-			ConstructFunction(Functions[i], bGenerateStubsOnly);
+			ConstructFunction(Functions[i], Dependencies, bGenerateStubsOnly);
 		}
 	}
 
@@ -170,7 +171,7 @@ void FBlueprintCompilerCppBackendBase::DeclareLocalVariables(FKismetFunctionCont
 	}
 }
 
-void FBlueprintCompilerCppBackendBase::ConstructFunction(FKismetFunctionContext& FunctionContext, bool bGenerateStubOnly)
+void FBlueprintCompilerCppBackendBase::ConstructFunction(FKismetFunctionContext& FunctionContext, const FGatherConvertedClassDependencies& Dependencies, bool bGenerateStubOnly)
 {
 	if (FunctionContext.IsDelegateSignature())
 	{
@@ -325,7 +326,7 @@ void FBlueprintCompilerCppBackendBase::ConstructFunction(FKismetFunctionContext&
 			}
 		}
 
-		const FString FunctionImplementation = InnerFunctionImplementation(FunctionContext, bUseSwitchState);
+		const FString FunctionImplementation = InnerFunctionImplementation(FunctionContext, Dependencies, bUseSwitchState);
 		Emit(Body, *FunctionImplementation);
 	}
 
@@ -374,7 +375,8 @@ void FBlueprintCompilerCppBackendBase::GenerateCodeFromStruct(UUserDefinedStruct
 {
 	check(SourceStruct);
 
-	EmitFileBeginning(SourceStruct->GetName(), SourceStruct);
+	FGatherConvertedClassDependencies Dependencies(SourceStruct);
+	EmitFileBeginning(SourceStruct->GetName(), &Dependencies);
 
 	const FString NewName = FString(TEXT("F")) + SourceStruct->GetName();
 	Emit(Header, TEXT("USTRUCT(BlueprintType"));
@@ -383,157 +385,86 @@ void FBlueprintCompilerCppBackendBase::GenerateCodeFromStruct(UUserDefinedStruct
 
 	Emit(Header, *FString::Printf(TEXT("struct %s\n{\npublic:\n\tGENERATED_BODY()\n"), *NewName));
 	EmitStructProperties(Header, SourceStruct);
-	Emit(Header, *FEmitDefaultValueHelper::GenerateGetDefaultValue(SourceStruct));
+	Emit(Header, *FEmitDefaultValueHelper::GenerateGetDefaultValue(SourceStruct, Dependencies));
 	Emit(Header, TEXT("};\n"));
 }
 
-void FBlueprintCompilerCppBackendBase::EmitFileBeginning(const FString& CleanName, UStruct* SourceStruct)
+void FBlueprintCompilerCppBackendBase::EmitFileBeginning(const FString& CleanName, FGatherConvertedClassDependencies* Dependencies)
 {
 	Emit(Header, TEXT("#pragma once\n\n"));
-	if (SourceStruct)
-	{ 
-		// find objects referenced by functions/script
-		TArray<UObject*> IncludeInHeader;
-		TArray<const UStruct*> DeclareInHeader;
-		TArray<UObject*> IncludeInBody;
-		{
-			FReferenceFinder HeaderReferenceFinder(IncludeInHeader, NULL, false, false, true, false);
-			FReferenceFinder BodyReferenceFinder(IncludeInBody, NULL, false, false, true, false);
-			{
-				TArray<UObject*> ObjectsToCheck;
-				GetObjectsWithOuter(SourceStruct, ObjectsToCheck, true);
-				for (auto Obj : ObjectsToCheck)
-				{
-					auto Property = Cast<UProperty>(Obj);
-					auto OwnerProperty = IsValid(Property) ? Property->GetOwnerProperty() : nullptr;
-					const bool bIsParam = OwnerProperty && (0 != (OwnerProperty->PropertyFlags & CPF_Parm));
-					const bool bIsMemberVariable = OwnerProperty && (OwnerProperty->GetOuter() == SourceStruct);
-					if (bIsParam || bIsMemberVariable)
-					{
-						const UObjectProperty* ObjectProperty = Cast<UObjectProperty>(OwnerProperty);
-						if (ObjectProperty)
-						{
-							DeclareInHeader.AddUnique(ObjectProperty->PropertyClass);
-							BodyReferenceFinder.FindReferences(Obj);
-						}
-						else
-						{
-							HeaderReferenceFinder.FindReferences(Obj);
-						}
-					}
-					else
-					{
-						BodyReferenceFinder.FindReferences(Obj);
-					}
-				}
-			}
 
-			if (auto BPGC = Cast<UClass>(SourceStruct))
-			{
-				BodyReferenceFinder.FindReferences(BPGC->GetDefaultObject(false));
-			}
+	auto EmitIncludeHeader = [&](FStringOutputDevice& Dst, const TCHAR* Message, bool bAddDotH)
+	{
+		Emit(Dst, *FString::Printf(TEXT("#include \"%s%s\"\n"), Message, bAddDotH ? TEXT(".h") : TEXT("")));
+	};
+	EmitIncludeHeader(Body, FApp::GetGameName(), true);
+	EmitIncludeHeader(Body, *CleanName, true);
+	EmitIncludeHeader(Body, TEXT("GeneratedCodeHelpers"), true);
 
-			if (auto SuperStruct = SourceStruct->GetSuperStruct())
-			{
-				IncludeInHeader.AddUnique(SuperStruct);
-			}
+	if (Dependencies)
+	{
+		TSet<FString> AlreadyIncluded;
+		AlreadyIncluded.Add(CleanName);
+		auto EmitInner = [&](FStringOutputDevice& Dst, const TSet<UField*>& Src, const TSet<UField*>& Declarations)
+	{
+		auto EngineSourceDir = FPaths::EngineSourceDir();
+		auto GameSourceDir = FPaths::GameSourceDir();
+
+		for (UField* Field : Src)
+		{
+			check(Field);
+			const bool bWantedType = Field->IsA<UBlueprintGeneratedClass>() || Field->IsA<UUserDefinedEnum>() || Field->IsA<UUserDefinedStruct>();
 
-			if (auto SourceClass = Cast<UClass>(SourceStruct))
+			// Wanted no-native type, thet will be converted
+			if (bWantedType)
 			{
-				for (auto& ImplementedInterface : SourceClass->Interfaces)
+				const FString Name = Field->GetName();
+				bool bAlreadyIncluded = false;
+				AlreadyIncluded.Add(Name, &bAlreadyIncluded);
+				if (!bAlreadyIncluded)
 				{
-					IncludeInHeader.AddUnique(ImplementedInterface.Class);
+					EmitIncludeHeader(Dst, *Name, true);
 				}
 			}
-		}
-
-		auto EmitIncludeHeader = [&](FStringOutputDevice& Dst, const TCHAR* Message, bool bAddDotH)
-		{
-			Emit(Dst, *FString::Printf(TEXT("#include \"%s%s\"\n"), Message, bAddDotH ? TEXT(".h") : TEXT("")));
-		};
-		EmitIncludeHeader(Body, FApp::GetGameName(), true);
-		EmitIncludeHeader(Body, *CleanName, true);
-		EmitIncludeHeader(Body, TEXT("GeneratedCodeHelpers"), true);
-
-		TSet<FString> AlreadyIncluded;
-		AlreadyIncluded.Add(SourceStruct->GetName());
-		auto EmitInner = [&](FStringOutputDevice& Dst, const TArray<UObject*>& Src, const TArray<const UStruct*>& Declarations )
-		{
-			auto EngineSourceDir = FPaths::EngineSourceDir();
-			auto GameSourceDir = FPaths::GameSourceDir();
-			auto CurrentPackage = SourceStruct->GetTypedOuter<UPackage>();
-
-			for (UObject* Obj : Src)
+			// headers for native items
+			else
 			{
-				bool bWantedType = Obj && (Obj->IsA<UBlueprintGeneratedClass>() || Obj->IsA<UUserDefinedEnum>() || Obj->IsA<UUserDefinedStruct>());
-				if (!bWantedType)
+				FString PackPath;
+				if (FSourceCodeNavigation::FindClassHeaderPath(Field, PackPath))
 				{
-					for (UObject* OuterObj = (Obj ? Obj->GetOuter() : nullptr); nullptr != OuterObj; OuterObj = OuterObj->GetOuter())
+					if (!PackPath.RemoveFromStart(EngineSourceDir))
 					{
-						if (OuterObj->IsA<UBlueprintGeneratedClass>() || OuterObj->IsA<UUserDefinedStruct>())
+						if (!PackPath.RemoveFromStart(GameSourceDir))
 						{
-							Obj = OuterObj;
-							bWantedType = true;
-							break;
+							PackPath = FPaths::GetCleanFilename(PackPath);
 						}
 					}
-				}
-
-				auto BPGC = (Obj && !bWantedType) ? Cast<UBlueprintGeneratedClass>(Obj->GetClass()) : nullptr;
-				if (BPGC)
-				{
-					bWantedType = true;
-					Obj = BPGC;
-				}
-
-				// Wanted no-native type, thet will be converted
-				if (bWantedType && Obj)
-				{
-					const FString Name = Obj->GetName();
 					bool bAlreadyIncluded = false;
-					AlreadyIncluded.Add(Name, &bAlreadyIncluded);
+					AlreadyIncluded.Add(PackPath, &bAlreadyIncluded);
 					if (!bAlreadyIncluded)
 					{
-						EmitIncludeHeader(Dst, *Name, true);
+						EmitIncludeHeader(Dst, *PackPath, false);
 					}
 				}
-				// headers for native items
-				else if (auto Field = Cast<UField>(Obj))
-				{
-					FString PackPath;
-					if ((Field->GetTypedOuter<UPackage>() != CurrentPackage)
-						&& FSourceCodeNavigation::FindClassHeaderPath(Field, PackPath))
-					{
-						if (!PackPath.RemoveFromStart(EngineSourceDir))
-						{
-							if (!PackPath.RemoveFromStart(GameSourceDir))
-							{
-								PackPath = FPaths::GetCleanFilename(PackPath);
-							}
-						}
-						bool bAlreadyIncluded = false;
-						AlreadyIncluded.Add(PackPath, &bAlreadyIncluded);
-						if (!bAlreadyIncluded)
-						{
-							EmitIncludeHeader(Dst, *PackPath, false);
-						}
-					}
-
-				}
 			}
+		}
 
-			Emit(Dst, TEXT("\n"));
+		Emit(Dst, TEXT("\n"));
 
-			for (const UStruct* Type : Declarations)
+		for (auto Type : Declarations)
+		{
+			if (auto ForwardDeclaredType = Cast<UClass>(Type))
 			{
-				Emit(Dst, *FString::Printf(TEXT("class %s;\n"), *(FString(Type->GetPrefixCPP()) + Type->GetName())));
+				Emit(Dst, *FString::Printf(TEXT("class %s;\n"), *(FString(ForwardDeclaredType->GetPrefixCPP()) + ForwardDeclaredType->GetName())));
 			}
+		}
 
-			Emit(Dst, TEXT("\n"));
-		};
+		Emit(Dst, TEXT("\n"));
+	};
 
-		EmitInner(Header, IncludeInHeader, DeclareInHeader);
-		EmitInner(Body, IncludeInBody, TArray<const UStruct*>() );
+		EmitInner(Header, Dependencies->IncludeInHeader, Dependencies->DeclareInHeader);
+		EmitInner(Body, Dependencies->IncludeInBody, TSet<UField*>());
 	}
+
 	Emit(Header, *FString::Printf(TEXT("#include \"%s.generated.h\"\n\n"), *CleanName));
 }
diff --git a/Engine/Source/Developer/BlueprintCompilerCppBackend/Private/BlueprintCompilerCppBackendBase.h b/Engine/Source/Developer/BlueprintCompilerCppBackend/Private/BlueprintCompilerCppBackendBase.h
--- a/Engine/Source/Developer/BlueprintCompilerCppBackend/Private/BlueprintCompilerCppBackendBase.h
+++ b/Engine/Source/Developer/BlueprintCompilerCppBackend/Private/BlueprintCompilerCppBackendBase.h
@@ -2,6 +2,7 @@
 #pragma  once
 
 #include "IBlueprintCompilerCppBackendModule.h"
+#include "BlueprintCompilerCppBackendGatherDependencies.h"
 
 /** The class generates all native code except of function body (notice InnerFunctionImplementation in not implemented) */
 class FBlueprintCompilerCppBackendBase : public IBlueprintCompilerCppBackend
@@ -69,7 +70,7 @@ protected:
 			}
 
 	*/
-	virtual FString InnerFunctionImplementation(FKismetFunctionContext& FunctionContext, bool bUseSwitchState) PURE_VIRTUAL(FBlueprintCompilerCppBackendBase::InnerFunctionImplementation, return FString(););
+	virtual FString InnerFunctionImplementation(FKismetFunctionContext& FunctionContext, const FGatherConvertedClassDependencies& Dependencies, bool bUseSwitchState) PURE_VIRTUAL(FBlueprintCompilerCppBackendBase::InnerFunctionImplementation, return FString(););
 
 	void EmitStructProperties(FStringOutputDevice& Target, UStruct* SourceClass);
 
@@ -77,9 +78,9 @@ protected:
 	void DeclareLocalVariables(FKismetFunctionContext& FunctionContext, TArray<UProperty*>& LocalVariables);
 	
 	/** Builds both the header declaration and body implementation of a function */
-	void ConstructFunction(FKismetFunctionContext& FunctionContext, bool bGenerateStubOnly);
+	void ConstructFunction(FKismetFunctionContext& FunctionContext, const FGatherConvertedClassDependencies& Dependencies, bool bGenerateStubOnly);
 
-	void EmitFileBeginning(const FString& CleanName, UStruct* SourceStruct);
+	void EmitFileBeginning(const FString& CleanName, FGatherConvertedClassDependencies* Dependencies);
 
 	int32 StatementToStateIndex(FKismetFunctionContext& FunctionContext, FBlueprintCompiledStatement* Statement)
 	{
diff --git a/Engine/Source/Developer/BlueprintCompilerCppBackend/Private/BlueprintCompilerCppBackendGatherDependencies.cpp b/Engine/Source/Developer/BlueprintCompilerCppBackend/Private/BlueprintCompilerCppBackendGatherDependencies.cpp
new file mode 100644
--- /dev/null
+++ b/Engine/Source/Developer/BlueprintCompilerCppBackend/Private/BlueprintCompilerCppBackendGatherDependencies.cpp
@@ -0,0 +1,156 @@
+// Copyright 1998-2015 Epic Games, Inc. All Rights Reserved.
+
+#include "BlueprintCompilerCppBackendModulePrivatePCH.h"
+#include "BlueprintCompilerCppBackendGatherDependencies.h"
+
+FGatherConvertedClassDependencies::FGatherConvertedClassDependencies(UStruct* InStruct) : OriginalStruct(InStruct)
+{
+	check(OriginalStruct);
+	DependenciesForHeader();
+	FindReferences(OriginalStruct);
+}
+
+bool FGatherConvertedClassDependencies::WillClassBeConverted(const UBlueprintGeneratedClass* InClass) const
+{
+	// TODO:
+
+	return InClass && !InClass->HasAnyFlags(RF_ClassDefaultObject);
+}
+
+void FGatherConvertedClassDependencies::DependenciesForHeader()
+{
+	TArray<UObject*> ObjectsToCheck;
+	GetObjectsWithOuter(OriginalStruct, ObjectsToCheck, true);
+
+	TArray<UObject*> NeededObjects;
+	FReferenceFinder HeaderReferenceFinder(NeededObjects, nullptr, false, false, true, false);
+
+	for (auto Obj : ObjectsToCheck)
+	{
+		auto Property = Cast<UProperty>(Obj);
+		auto OwnerProperty = IsValid(Property) ? Property->GetOwnerProperty() : nullptr;
+		const bool bIsParam = OwnerProperty && (0 != (OwnerProperty->PropertyFlags & CPF_Parm)) && OwnerProperty->IsIn(OriginalStruct);
+		const bool bIsMemberVariable = OwnerProperty && (OwnerProperty->GetOuter() == OriginalStruct);
+		if (bIsParam || bIsMemberVariable)
+		{
+			const UObjectProperty* ObjectProperty = Cast<UObjectProperty>(OwnerProperty);
+			if (ObjectProperty)
+			{
+				DeclareInHeader.Add(ObjectProperty->PropertyClass);
+			}
+			else
+			{
+				HeaderReferenceFinder.FindReferences(Obj);
+			}
+		}
+	}
+
+	if (auto SuperStruct = OriginalStruct->GetSuperStruct())
+	{
+		IncludeInHeader.Add(SuperStruct);
+	}
+
+	if (auto SourceClass = Cast<UClass>(OriginalStruct))
+	{
+		for (auto& ImplementedInterface : SourceClass->Interfaces)
+		{
+			IncludeInHeader.Add(ImplementedInterface.Class);
+		}
+	}
+
+	for (auto Obj : NeededObjects)
+	{
+		if (ShouldIncludeHeaderFor(Obj))
+		{
+			IncludeInHeader.Add(CastChecked<UField>(Obj));
+		}
+	}
+}
+
+bool FGatherConvertedClassDependencies::ShouldIncludeHeaderFor(UObject* Obj) const
+{
+	if (Obj
+		&& (Obj->IsA<UClass>() || Obj->IsA<UEnum>() || Obj->IsA<UScriptStruct>())
+		&& !Obj->HasAnyFlags(RF_ClassDefaultObject))
+	{
+		auto ObjAsBPGC = Cast<UBlueprintGeneratedClass>(Obj);
+		const bool bWillBeConvetedAsBPGC = ObjAsBPGC && WillClassBeConverted(ObjAsBPGC);
+		const bool bRemainAsUnconvertedBPGC = ObjAsBPGC && !bWillBeConvetedAsBPGC;
+		if (!bRemainAsUnconvertedBPGC && (Obj->GetOutermost() != OriginalStruct->GetOutermost()))
+		{
+			return true;
+		}
+	}
+	return false;
+}
+
+void FGatherConvertedClassDependencies::HandleObjectReference(UObject*& InObject, const UObject* InReferencingObject, const UProperty* InReferencingProperty)
+{
+	if (!InObject || InObject->IsA<UBlueprint>() || (InObject == OriginalStruct))
+	{
+		return;
+	}
+
+	auto ObjAsField = Cast<UField>(InObject);
+	if (ObjAsField && ShouldIncludeHeaderFor(ObjAsField) && !IncludeInHeader.Contains(ObjAsField))
+	{
+		IncludeInBody.Add(ObjAsField);
+	}
+
+	//TODO: What About Delegates?
+	auto ObjAsBPGC = Cast<UBlueprintGeneratedClass>(InObject);
+	const bool bWillBeConvetedAsBPGC = ObjAsBPGC && WillClassBeConverted(ObjAsBPGC);
+	const bool bRemainAsUnconvertedBPGC = ObjAsBPGC && !bWillBeConvetedAsBPGC;
+	if (bWillBeConvetedAsBPGC)
+	{
+		ConvertedClasses.Add(ObjAsBPGC);
+	}
+	else if (InObject->IsA<UUserDefinedStruct>())
+	{
+		if (!InObject->HasAnyFlags(RF_ClassDefaultObject))
+		{
+			ConvertedStructs.Add(CastChecked<UUserDefinedStruct>(InObject));
+		}
+	}
+	else if (InObject->IsA<UUserDefinedEnum>())
+	{
+		if (!InObject->HasAnyFlags(RF_ClassDefaultObject))
+		{
+			ConvertedEnum.Add(CastChecked<UUserDefinedEnum>(InObject));
+		}
+	}
+	else if (InObject->IsAsset() && !InObject->IsIn(OriginalStruct))
+	{
+		Assets.Add(InObject);
+		return;
+	}
+	else if (auto ObjAsClass = Cast<UClass>(InObject))
+	{
+		if (ObjAsClass->HasAnyClassFlags(CLASS_Native))
+		{
+			return;
+		}
+	}
+	else if (InObject->IsA<UScriptStruct>())
+	{
+		return;
+	}
+
+	bool AlreadyAdded = false;
+	SerializedObjects.Add(InObject, &AlreadyAdded);
+	if (!AlreadyAdded)
+	{
+		FindReferences(InObject);
+	}
+}
+
+void FGatherConvertedClassDependencies::FindReferences(UObject* Object)
+{
+	{
+		FSimpleObjectReferenceCollectorArchive CollectorArchive(Object, *this);
+		CollectorArchive.SetSerializedProperty(nullptr);
+		CollectorArchive.SetFilterEditorOnly(true);
+		Object->SerializeScriptProperties(CollectorArchive);
+	}
+	Object->CallAddReferencedObjects(*this);
+}
\ No newline at end of file
diff --git a/Engine/Source/Developer/BlueprintCompilerCppBackend/Private/BlueprintCompilerCppBackendUtils.h b/Engine/Source/Developer/BlueprintCompilerCppBackend/Private/BlueprintCompilerCppBackendUtils.h
--- a/Engine/Source/Developer/BlueprintCompilerCppBackend/Private/BlueprintCompilerCppBackendUtils.h
+++ b/Engine/Source/Developer/BlueprintCompilerCppBackend/Private/BlueprintCompilerCppBackendUtils.h
@@ -1,5 +1,8 @@
 // Copyright 1998-2015 Epic Games, Inc. All Rights Reserved.
 #pragma  once
+
+#include "BlueprintCompilerCppBackendGatherDependencies.h"
+
 struct FEmitterLocalContext
 {
 private:
@@ -14,15 +17,17 @@ private:
 	TMap<UObject*, FString> NativeObjectNamesInConstructor;
 	TArray<UObject*> ObjectsCreatedPerClass;
 
+	const FGatherConvertedClassDependencies& Dependencies;
 
 public:
 
 	bool bCreatingObjectsPerClass;
 	bool bInsideConstructor;
 
-	FEmitterLocalContext()
+	FEmitterLocalContext(UClass* InClass, const FGatherConvertedClassDependencies& InDependencies)
 		: LocalNameIndexMax(0)
-		, ActualClass(nullptr)
+		, ActualClass(InClass)
+		, Dependencies(InDependencies)
 		, bCreatingObjectsPerClass(false)
 		, bInsideConstructor(false)
 	{}
@@ -85,11 +90,6 @@ public:
 		return FString::Printf(TEXT("%s->FindPropertyByName(FName(TEXT(\"%s\")))"), *OwnerPath, *Property->GetNameCPP());
 	}
 
-	void SetCurrentlyGeneratedClass(UClass* InClass)
-	{
-		ActualClass = InClass;
-	}
-
 	FString GenerateUniqueLocalName()
 	{
 		const FString UniqueNameBase = TEXT("__Local__");
@@ -103,14 +103,6 @@ public:
 		return ActualClass;
 	}
 
-	bool WillClassBeConverted(UClass* InClass) const
-	{
-		// TODO:
-
-		auto BPGC = Cast<UBlueprintGeneratedClass>(InClass);
-		return (nullptr != BPGC) && !InClass->HasAnyFlags(RF_ClassDefaultObject);
-	}
-
 	/** All objects (that can be referenced from other package) that will have a different path in cooked build 
 	(due to the native code generation), should be handled by this function */
 	FString FindGloballyMappedObject(UObject* Object, bool bLoadIfNotFound = false)
@@ -124,7 +116,8 @@ public:
 
 		if (auto ObjClass = Cast<UClass>(Object))
 		{
-			if (ObjClass->HasAnyClassFlags(CLASS_Native) || WillClassBeConverted(ObjClass))
+			auto BPGC = Cast<UBlueprintGeneratedClass>(ObjClass);
+			if (ObjClass->HasAnyClassFlags(CLASS_Native) || (BPGC && Dependencies.WillClassBeConverted(BPGC)))
 			{
 				return FString::Printf(TEXT("%s%s::StaticClass()"), ObjClass->GetPrefixCPP(), *ObjClass->GetName());
 			}
@@ -148,7 +141,7 @@ public:
 		}
 
 		int32 ObjectsCreatedPerClassIdx = INDEX_NONE;
-		if (Object && ObjectsCreatedPerClass.Find(Object, ObjectsCreatedPerClassIdx))
+		if (ActualClass && Object && ObjectsCreatedPerClass.Find(Object, ObjectsCreatedPerClassIdx))
 		{
 			return FString::Printf(TEXT("CastChecked<%s%s>(%s%s::StaticClass()->ConvertedSubobjectsFromBPGC[%d])")
 				, Object->GetClass()->GetPrefixCPP()
@@ -266,9 +259,9 @@ struct FEmitHelper
 
 struct FEmitDefaultValueHelper
 {
-	static FString GenerateGetDefaultValue(const UUserDefinedStruct* Struct);
+	static FString GenerateGetDefaultValue(const UUserDefinedStruct* Struct, const FGatherConvertedClassDependencies& Dependencies);
 
-	static FString GenerateConstructor(UClass* BPGC);
+	static FString GenerateConstructor(UClass* BPGC, const FGatherConvertedClassDependencies& Dependencies);
 
 	enum class EPropertyAccessOperator
 	{
diff --git a/Engine/Source/Developer/BlueprintCompilerCppBackend/Private/BlueprintCompilerCppBackendValueHelper.cpp b/Engine/Source/Developer/BlueprintCompilerCppBackend/Private/BlueprintCompilerCppBackendValueHelper.cpp
--- a/Engine/Source/Developer/BlueprintCompilerCppBackend/Private/BlueprintCompilerCppBackendValueHelper.cpp
+++ b/Engine/Source/Developer/BlueprintCompilerCppBackend/Private/BlueprintCompilerCppBackendValueHelper.cpp
@@ -78,7 +78,7 @@ void FEmitDefaultValueHelper::OuterGenerate(FEmitterLocalContext& Context
 	}
 }
 
-FString FEmitDefaultValueHelper::GenerateGetDefaultValue(const UUserDefinedStruct* Struct)
+FString FEmitDefaultValueHelper::GenerateGetDefaultValue(const UUserDefinedStruct* Struct, const FGatherConvertedClassDependencies& Dependencies)
 {
 	check(Struct);
 	const FString StructName = FString(TEXT("F")) + Struct->GetName();
@@ -87,7 +87,7 @@ FString FEmitDefaultValueHelper::GenerateGetDefaultValue(const UUserDefinedStruc
 	FStructOnScope StructData(Struct);
 	FStructureEditorUtils::Fill_MakeStructureDefaultValue(Struct, StructData.GetStructMemory());
 
-	FEmitterLocalContext Context;
+	FEmitterLocalContext Context(nullptr, Dependencies);
 	Context.IncreaseIndent();
 	Context.IncreaseIndent();
 	for (auto Property : TFieldRange<const UProperty>(Struct))
@@ -323,103 +323,57 @@ FString FEmitDefaultValueHelper::HandleNonNativeComponent(FEmitterLocalContext&
 	return VariableName;
 }
 
-struct FGatherDependencies : public FReferenceCollector
+struct FDependenciesHelper
 {
-	TSet<UObject*> SerializedObjects;
-	TSet<UObject*> Assets;
-	TSet<UStruct*> ConvertedClasses;
-
-	FEmitterLocalContext& Context;
-	UClass* OriginalClass;
-
-	void FindReferences(UObject* Object)
+	static void AddDependenciesInConstructor(FEmitterLocalContext& Context, const FGatherConvertedClassDependencies& GatherDependencies)
 	{
+		if (GatherDependencies.ConvertedClasses.Num())
 		{
-			FSimpleObjectReferenceCollectorArchive CollectorArchive(Object, *this);
-			CollectorArchive.SetSerializedProperty(nullptr);
-			CollectorArchive.SetFilterEditorOnly(true);
-			Object->SerializeScriptProperties(CollectorArchive);
+			Context.AddLine(TEXT("// List of all referenced converted classes"));
 		}
-		Object->CallAddReferencedObjects(*this);
-	}
-
-	virtual bool IsIgnoringArchetypeRef() const override { return false; }
-	virtual bool IsIgnoringTransient() const override { return true; }
-
-	virtual void HandleObjectReference(UObject*& InObject, const UObject* InReferencingObject, const UProperty* InReferencingProperty) override
-	{
-		if (InObject)
+		for (auto LocStruct : GatherDependencies.ConvertedClasses)
 		{
-			if (InObject->IsA<UBlueprint>() || (InObject == OriginalClass))
-			{
-				return;
-			}
-
-			// converted class
-			//TODO: What About Delegates?
-			auto ObjAsBPGC = Cast<UBlueprintGeneratedClass>(InObject);
-			const bool SaveAsConvertedClass = ObjAsBPGC && Context.WillClassBeConverted(ObjAsBPGC);
-			auto ObjAsUDS = Cast<UUserDefinedStruct>(InObject);// all UDS must be converted
-			const bool SaveAsConvertedStruct = ObjAsUDS && !ObjAsUDS->HasAnyFlags(RF_ClassDefaultObject);
-			if (SaveAsConvertedClass || SaveAsConvertedStruct)
-			{
-				ConvertedClasses.Add(CastChecked<UStruct>(InObject));
-			}
-			// asset
-			else if (InObject->IsAsset() 
-				&& !InObject->IsIn(OriginalClass) 
-				&& (InObject != OriginalClass)
-				&& !InObject->IsA<UUserDefinedEnum>()) // all UDE must be converted
-			{
-				Assets.Add(InObject);
-				return;
-			}
-
-			bool AlreadyAdded = false;
-			SerializedObjects.Add(InObject, &AlreadyAdded);
-			if (!AlreadyAdded)
-			{
-				FindReferences(InObject);
-			}
+			Context.AddLine(FString::Printf(TEXT("GetClass()->ConvertedSubobjectsFromBPGC.Add(%s%s::StaticClass());")
+				, LocStruct->GetPrefixCPP()
+				, *LocStruct->GetName()));
 		}
-	}
-
-	FGatherDependencies(UClass* InClass, FEmitterLocalContext& InContext) 
-		: Context(InContext)
-		, OriginalClass(InClass)
-	{
-		check(OriginalClass);
-		FindReferences(OriginalClass);
-	}
 
-	void AddDependenciesInConstructor()
-	{
-		Context.AddLine(TEXT("// List of all referenced converted classes"));
-		for (auto LocStruct : ConvertedClasses)
+		if (GatherDependencies.ConvertedStructs.Num())
+		{
+			Context.AddLine(TEXT("// List of all referenced converted structures"));
+		}
+		for (auto LocStruct : GatherDependencies.ConvertedStructs)
 		{
-			const bool bIsClass = nullptr != Cast<UClass>(LocStruct);
-			Context.AddLine(FString::Printf(TEXT("GetClass()->ConvertedSubobjectsFromBPGC.Add(%s%s::%s());")
+			Context.AddLine(FString::Printf(TEXT("GetClass()->ConvertedSubobjectsFromBPGC.Add(%s%s::StaticStruct());")
 				, LocStruct->GetPrefixCPP()
-				, *LocStruct->GetName()
-				, (bIsClass ? TEXT("StaticClass") : TEXT("StaticStruct"))));
+				, *LocStruct->GetName()));
 		}
-		Context.AddLine(TEXT("// List of all referenced assets"));
-		for (auto LocAsset : Assets)
+
+		if (GatherDependencies.Assets.Num())
+		{
+			Context.AddLine(TEXT("// List of all referenced assets"));
+		}
+		for (auto LocAsset : GatherDependencies.Assets)
 		{
 			const FString AssetStr = Context.FindGloballyMappedObject(LocAsset, true);
 			Context.AddLine(FString::Printf(TEXT("GetClass()->ConvertedSubobjectsFromBPGC.Add(%s);"), *AssetStr));
 		}
 	}
 
-	void AddStaticFunctionsForDependencies()
+	static void AddStaticFunctionsForDependencies(FEmitterLocalContext& Context, const FGatherConvertedClassDependencies& GatherDependencies)
 	{
+		auto OriginalClass = Context.GetCurrentlyGeneratedClass();
 		const FString CppClassName = FString::Printf(TEXT("%s%s"), OriginalClass->GetPrefixCPP(), *OriginalClass->GetName());
 		// __StaticDependenciesConvertedClasses
 		Context.AddLine(FString::Printf(TEXT("void %s::__StaticDependenciesConvertedClasses(TArray<FName>& OutNames)"), *CppClassName));
 		Context.AddLine(TEXT("{"));
 		Context.IncreaseIndent();
 
-		for (auto LocClass : ConvertedClasses)
+		for (auto LocClass : GatherDependencies.ConvertedClasses)
+		{
+			Context.AddLine(FString::Printf(TEXT("OutNames.Add(TEXT(\"%s\"));"), *LocClass->GetName()));
+		}
+		for (auto LocClass : GatherDependencies.ConvertedStructs)
 		{
 			Context.AddLine(FString::Printf(TEXT("OutNames.Add(TEXT(\"%s\"));"), *LocClass->GetName()));
 		}
@@ -432,7 +386,7 @@ struct FGatherDependencies : public FReferenceCollector
 		Context.AddLine(TEXT("{"));
 		Context.IncreaseIndent();
 
-		for (UObject* LocAsset : Assets)
+		for (UObject* LocAsset : GatherDependencies.Assets)
 		{
 			const FString PakagePath = LocAsset->GetOutermost()->GetPathName();
 			Context.AddLine(FString::Printf(TEXT("OutPackagePaths.Add(TEXT(\"%s\"));"), *PakagePath));
@@ -467,20 +421,17 @@ struct FGatherDependencies : public FReferenceCollector
 	}
 };
 
-FString FEmitDefaultValueHelper::GenerateConstructor(UClass* InBPGC)
+FString FEmitDefaultValueHelper::GenerateConstructor(UClass* InBPGC, const FGatherConvertedClassDependencies& Dependencies)
 {
 	auto BPGC = CastChecked<UBlueprintGeneratedClass>(InBPGC);
 	const FString CppClassName = FString(BPGC->GetPrefixCPP()) + BPGC->GetName();
 
-	FEmitterLocalContext Context;
-	Context.SetCurrentlyGeneratedClass(BPGC);
+	FEmitterLocalContext Context(BPGC, Dependencies);
 	Context.AddLine(FString::Printf(TEXT("%s::%s(const FObjectInitializer& ObjectInitializer) : Super(ObjectInitializer)"), *CppClassName, *CppClassName));
 	Context.AddLine(TEXT("{"));
 	Context.IncreaseIndent();
 	Context.bInsideConstructor = true;
 
-	FGatherDependencies GatherDependencies(InBPGC, Context);
-
 	// When CDO is created create all subobjects owned by the class
 	TArray<UActorComponent*> ActorComponentTempatesOwnedByClass = BPGC->ComponentTemplates;
 	{
@@ -536,7 +487,7 @@ FString FEmitDefaultValueHelper::GenerateConstructor(UClass* InBPGC)
 
 		Context.bCreatingObjectsPerClass = false;
 
-		GatherDependencies.AddDependenciesInConstructor();
+		FDependenciesHelper::AddDependenciesInConstructor(Context, Dependencies);
 
 		Context.DecreaseIndent();
 		Context.AddLine(TEXT("}"));
@@ -655,7 +606,7 @@ FString FEmitDefaultValueHelper::GenerateConstructor(UClass* InBPGC)
 	Context.AddLine(TEXT("}"));
 	Context.bInsideConstructor = false;
 
-	GatherDependencies.AddStaticFunctionsForDependencies();
+	FDependenciesHelper::AddStaticFunctionsForDependencies(Context, Dependencies);
 
 	FBackendHelperUMG::EmitWidgetInitializationFunctions(Context);
 
diff --git a/Engine/Source/Developer/BlueprintCompilerCppBackend/Public/BlueprintCompilerCppBackendGatherDependencies.h b/Engine/Source/Developer/BlueprintCompilerCppBackend/Public/BlueprintCompilerCppBackendGatherDependencies.h
new file mode 100644
--- /dev/null
+++ b/Engine/Source/Developer/BlueprintCompilerCppBackend/Public/BlueprintCompilerCppBackendGatherDependencies.h
@@ -0,0 +1,39 @@
+// Copyright 1998-2015 Epic Games, Inc. All Rights Reserved.
+
+#pragma  once
+
+#include "Engine/BlueprintGeneratedClass.h"
+#include "Engine/UserDefinedStruct.h"
+#include "Engine/UserDefinedEnum.h"
+
+struct BLUEPRINTCOMPILERCPPBACKEND_API FGatherConvertedClassDependencies : public FReferenceCollector
+{
+protected:
+	TSet<UObject*> SerializedObjects;
+	UStruct* OriginalStruct;
+
+public:
+	TSet<UObject*> Assets;
+	TSet<UBlueprintGeneratedClass*> ConvertedClasses;
+	TSet<UUserDefinedStruct*> ConvertedStructs;
+	TSet<UUserDefinedEnum*> ConvertedEnum;
+
+	TSet<UField*> IncludeInHeader;
+	TSet<UField*> DeclareInHeader;
+	TSet<UField*> IncludeInBody;
+public:
+	FGatherConvertedClassDependencies(UStruct* InStruct);
+
+	virtual bool IsIgnoringArchetypeRef() const override { return false; }
+	virtual bool IsIgnoringTransient() const override { return true; }
+
+	virtual void HandleObjectReference(UObject*& InObject, const UObject* InReferencingObject, const UProperty* InReferencingProperty) override;
+
+public:
+	bool WillClassBeConverted(const UBlueprintGeneratedClass* InClass) const;
+
+protected:
+	void FindReferences(UObject* Object);
+	void DependenciesForHeader();
+	bool ShouldIncludeHeaderFor(UObject* Object) const;
+};
diff --git a/Engine/Source/Developer/BlueprintNativeCodeGen/BlueprintNativeCodeGen.Build.cs b/Engine/Source/Developer/BlueprintNativeCodeGen/BlueprintNativeCodeGen.Build.cs
--- a/Engine/Source/Developer/BlueprintNativeCodeGen/BlueprintNativeCodeGen.Build.cs
+++ b/Engine/Source/Developer/BlueprintNativeCodeGen/BlueprintNativeCodeGen.Build.cs
@@ -27,6 +27,7 @@ public class BlueprintNativeCodeGen : ModuleRules
                 "KismetCompiler",
                 "Json",
                 "JsonUtilities",
+                "BlueprintCompilerCppBackend",
 			}
 		);
 
diff --git a/Engine/Source/Developer/BlueprintNativeCodeGen/Private/NativeCodeGenerationTool.cpp b/Engine/Source/Developer/BlueprintNativeCodeGen/Private/NativeCodeGenerationTool.cpp
--- a/Engine/Source/Developer/BlueprintNativeCodeGen/Private/NativeCodeGenerationTool.cpp
+++ b/Engine/Source/Developer/BlueprintNativeCodeGen/Private/NativeCodeGenerationTool.cpp
@@ -8,7 +8,9 @@
 #include "Engine/UserDefinedStruct.h"
 #include "SourceCodeNavigation.h"
 #include "DesktopPlatformModule.h" // for InvalidateMakefiles()
+#include "BlueprintNativeCodeGenUtils.h"
 //#include "Editor/KismetCompiler/Public/BlueprintCompilerCppBackendInterface.h"
+#include "Developer/BlueprintCompilerCppBackend/Public/BlueprintCompilerCppBackendGatherDependencies.h"
 
 #define LOCTEXT_NAMESPACE "NativeCodeGenerationTool"
 
@@ -33,57 +35,29 @@ struct FGeneratedCodeData
 	FString ErrorString;
 	FString ClassName;
 	TWeakObjectPtr<UBlueprint> Blueprint;
-	TSet<UObject*> DependentObjects;
+	TSet<UField*> DependentObjects;
 
 	void GatherUserDefinedDependencies(UBlueprint& InBlueprint)
 	{
-		TArray<UObject*> ReferencedObjects;
+		FGatherConvertedClassDependencies ClassDependencies(InBlueprint.GeneratedClass);
+		for (auto Iter : ClassDependencies.ConvertedClasses)
 		{
-			FReferenceFinder ReferenceFinder(ReferencedObjects, NULL, false, false, false, false);
-
-			{
-				TArray<UObject*> ObjectsToCheck;
-				GetObjectsWithOuter(InBlueprint.GeneratedClass, ObjectsToCheck, true);
-				for (auto Obj : ObjectsToCheck)
-				{
-					if (IsValid(Obj))
-					{
-						ReferenceFinder.FindReferences(Obj);
-					}
+			DependentObjects.Add(Iter);
 				}
-			}
-
-			if (InBlueprint.GeneratedClass)
-			{
-				ReferenceFinder.FindReferences(InBlueprint.GeneratedClass->GetDefaultObject(false));
-			}
-
-			for (UClass* Class = InBlueprint.GeneratedClass->GetSuperClass(); Class && !Class->HasAnyClassFlags(CLASS_Native); Class = Class->GetSuperClass())
+		for (auto Iter : ClassDependencies.ConvertedStructs)
 			{
-				ReferencedObjects.Add(Class);
+			DependentObjects.Add(Iter);
 			}
-
-			for (auto& ImplementedInterface : InBlueprint.GeneratedClass->Interfaces)
+		for (auto Iter : ClassDependencies.ConvertedEnum)
 			{
-				if (ImplementedInterface.Class && !ImplementedInterface.Class->HasAnyClassFlags(CLASS_Native))
-				{
-					ReferencedObjects.Add(ImplementedInterface.Class);
-				}
-			}
+			DependentObjects.Add(Iter);
 		}
 
 		TypeDependencies.Empty();
-		for (auto Obj : ReferencedObjects)
-		{
-			if (IsValid(Obj) && !Obj->IsIn(InBlueprint.GeneratedClass) && (Obj != InBlueprint.GeneratedClass))
-			{
-				if (Obj->IsA<UBlueprintGeneratedClass>() || Obj->IsA<UUserDefinedEnum>() || Obj->IsA<UUserDefinedStruct>())
+		for (auto Obj : DependentObjects)
 				{
 					TypeDependencies += Obj->GetPathName();
 					TypeDependencies += TEXT("\n");
-					DependentObjects.Add(Obj);
-				}
-			}
 		}
 		DependentObjects.Add(InBlueprint.GeneratedClass);
 
