Comments added

diff --git a/Engine/Source/Developer/BlueprintCompilerCppBackend/Public/BlueprintCompilerCppBackendGatherDependencies.h b/Engine/Source/Developer/BlueprintCompilerCppBackend/Public/BlueprintCompilerCppBackendGatherDependencies.h
--- a/Engine/Source/Developer/BlueprintCompilerCppBackend/Public/BlueprintCompilerCppBackendGatherDependencies.h
+++ b/Engine/Source/Developer/BlueprintCompilerCppBackend/Public/BlueprintCompilerCppBackendGatherDependencies.h
@@ -6,6 +6,7 @@
 #include "Engine/UserDefinedStruct.h"
 #include "Engine/UserDefinedEnum.h"
 
+/** The struct gathers dependencies of a converted BPGC */
 struct BLUEPRINTCOMPILERCPPBACKEND_API FGatherConvertedClassDependencies : public FReferenceCollector
 {
 protected:
@@ -13,11 +14,13 @@ protected:
 	UStruct* OriginalStruct;
 
 public:
+	// Dependencies:
 	TSet<UObject*> Assets;
 	TSet<UBlueprintGeneratedClass*> ConvertedClasses;
 	TSet<UUserDefinedStruct*> ConvertedStructs;
 	TSet<UUserDefinedEnum*> ConvertedEnum;
 
+	// What to include/declare in the generated code:
 	TSet<UField*> IncludeInHeader;
 	TSet<UField*> DeclareInHeader;
 	TSet<UField*> IncludeInBody;
