UEBP-40 C++ code generation (WIP):
- improved (working) code for UMG Widgets
- various fixes

diff --git a/Engine/Source/Developer/BlueprintCompilerCppBackend/Private/BlueprintCompilerCppBackendBase.cpp b/Engine/Source/Developer/BlueprintCompilerCppBackend/Private/BlueprintCompilerCppBackendBase.cpp
--- a/Engine/Source/Developer/BlueprintCompilerCppBackend/Private/BlueprintCompilerCppBackendBase.cpp
+++ b/Engine/Source/Developer/BlueprintCompilerCppBackend/Private/BlueprintCompilerCppBackendBase.cpp
@@ -14,7 +14,7 @@ void FBlueprintCompilerCppBackendBase::EmitStructProperties(FStringOutputDevice&
 
 		Emit(Header, TEXT("\n\tUPROPERTY("));
 		{
-			TArray<FString> Tags = FEmitHelper::ProperyFlagsToTags(Property->PropertyFlags);
+			TArray<FString> Tags = FEmitHelper::ProperyFlagsToTags(Property->PropertyFlags, nullptr != Cast<UClass>(SourceClass));
 			Tags.Emplace(FEmitHelper::HandleRepNotifyFunc(Property));
 			Tags.Emplace(FEmitHelper::HandleMetaData(Property, false));
 			Tags.Remove(FString());
@@ -403,64 +403,66 @@ void FBlueprintCompilerCppBackendBase::EmitFileBeginning(const FString& CleanNam
 
 	if (Dependencies)
 	{
+		Emit(Body, *FBackendHelperUMG::AdditionalHeaderIncludeForWidget(Cast<UClass>(Dependencies->GetOriginalStruct())));
+
 		TSet<FString> AlreadyIncluded;
 		AlreadyIncluded.Add(CleanName);
 		auto EmitInner = [&](FStringOutputDevice& Dst, const TSet<UField*>& Src, const TSet<UField*>& Declarations)
-	{
-		auto EngineSourceDir = FPaths::EngineSourceDir();
-		auto GameSourceDir = FPaths::GameSourceDir();
-
-		for (UField* Field : Src)
 		{
-			check(Field);
-			const bool bWantedType = Field->IsA<UBlueprintGeneratedClass>() || Field->IsA<UUserDefinedEnum>() || Field->IsA<UUserDefinedStruct>();
+			auto EngineSourceDir = FPaths::EngineSourceDir();
+			auto GameSourceDir = FPaths::GameSourceDir();
 
-			// Wanted no-native type, thet will be converted
-			if (bWantedType)
+			for (UField* Field : Src)
 			{
-				const FString Name = Field->GetName();
-				bool bAlreadyIncluded = false;
-				AlreadyIncluded.Add(Name, &bAlreadyIncluded);
-				if (!bAlreadyIncluded)
+				check(Field);
+				const bool bWantedType = Field->IsA<UBlueprintGeneratedClass>() || Field->IsA<UUserDefinedEnum>() || Field->IsA<UUserDefinedStruct>();
+
+				// Wanted no-native type, thet will be converted
+				if (bWantedType)
 				{
-					EmitIncludeHeader(Dst, *Name, true);
+					const FString Name = Field->GetName();
+					bool bAlreadyIncluded = false;
+					AlreadyIncluded.Add(Name, &bAlreadyIncluded);
+					if (!bAlreadyIncluded)
+					{
+						EmitIncludeHeader(Dst, *Name, true);
+					}
 				}
-			}
-			// headers for native items
-			else
-			{
-				FString PackPath;
-				if (FSourceCodeNavigation::FindClassHeaderPath(Field, PackPath))
+				// headers for native items
+				else
 				{
-					if (!PackPath.RemoveFromStart(EngineSourceDir))
+					FString PackPath;
+					if (FSourceCodeNavigation::FindClassHeaderPath(Field, PackPath))
 					{
-						if (!PackPath.RemoveFromStart(GameSourceDir))
+						if (!PackPath.RemoveFromStart(EngineSourceDir))
 						{
-							PackPath = FPaths::GetCleanFilename(PackPath);
+							if (!PackPath.RemoveFromStart(GameSourceDir))
+							{
+								PackPath = FPaths::GetCleanFilename(PackPath);
+							}
+						}
+						bool bAlreadyIncluded = false;
+						AlreadyIncluded.Add(PackPath, &bAlreadyIncluded);
+						if (!bAlreadyIncluded)
+						{
+							EmitIncludeHeader(Dst, *PackPath, false);
 						}
-					}
-					bool bAlreadyIncluded = false;
-					AlreadyIncluded.Add(PackPath, &bAlreadyIncluded);
-					if (!bAlreadyIncluded)
-					{
-						EmitIncludeHeader(Dst, *PackPath, false);
 					}
 				}
 			}
-		}
 
-		Emit(Dst, TEXT("\n"));
+			Emit(Dst, TEXT("\n"));
 
-		for (auto Type : Declarations)
-		{
-			if (auto ForwardDeclaredType = Cast<UClass>(Type))
+			for (auto Type : Declarations)
 			{
-				Emit(Dst, *FString::Printf(TEXT("class %s;\n"), *(FString(ForwardDeclaredType->GetPrefixCPP()) + ForwardDeclaredType->GetName())));
+				if (auto ForwardDeclaredType = Cast<UClass>(Type))
+				{
+					Emit(Dst, *FString::Printf(TEXT("class %s;\n"), *(FString(ForwardDeclaredType->GetPrefixCPP()) + ForwardDeclaredType->GetName())));
+				}
 			}
-		}
 
-		Emit(Dst, TEXT("\n"));
-	};
+			Emit(Dst, TEXT("\n"));
+		};
 
 		EmitInner(Header, Dependencies->IncludeInHeader, Dependencies->DeclareInHeader);
 		EmitInner(Body, Dependencies->IncludeInBody, TSet<UField*>());
diff --git a/Engine/Source/Developer/BlueprintCompilerCppBackend/Private/BlueprintCompilerCppBackendUMG.cpp b/Engine/Source/Developer/BlueprintCompilerCppBackend/Private/BlueprintCompilerCppBackendUMG.cpp
--- a/Engine/Source/Developer/BlueprintCompilerCppBackend/Private/BlueprintCompilerCppBackendUMG.cpp
+++ b/Engine/Source/Developer/BlueprintCompilerCppBackend/Private/BlueprintCompilerCppBackendUMG.cpp
@@ -15,6 +15,15 @@ FString FBackendHelperUMG::WidgetFunctionsInHeader(UClass* SourceClass)
 	return FString();
 }
 
+FString FBackendHelperUMG::AdditionalHeaderIncludeForWidget(UClass* SourceClass)
+{
+	if (Cast<UWidgetBlueprintGeneratedClass>(SourceClass))
+	{
+		return FString(TEXT("\n#include \"Runtime/UMG/Public/UMG.h\"\n"));
+	}
+	return FString();
+}
+
 void FBackendHelperUMG::CreateClassSubobjects(FEmitterLocalContext& Context)
 {
 	if (auto WidgetClass = Cast<UWidgetBlueprintGeneratedClass>(Context.GetCurrentlyGeneratedClass()))
@@ -27,7 +36,7 @@ void FBackendHelperUMG::CreateClassSubobjects(FEmitterLocalContext& Context)
 		for (auto Anim : WidgetClass->Animations)
 		{
 			ensure(Anim->GetOuter() == Context.GetCurrentlyGeneratedClass());
-			FEmitDefaultValueHelper::HandleClassSubobject(Context, WidgetClass->WidgetTree);
+			FEmitDefaultValueHelper::HandleClassSubobject(Context, Anim);
 		}
 	}
 }
@@ -96,3 +105,4 @@ void FBackendHelperUMG::EmitWidgetInitializationFunctions(FEmitterLocalContext&
 		Context.AddLine(TEXT("}"));
 	}
 }
+
diff --git a/Engine/Source/Developer/BlueprintCompilerCppBackend/Private/BlueprintCompilerCppBackendUtils.cpp b/Engine/Source/Developer/BlueprintCompilerCppBackend/Private/BlueprintCompilerCppBackendUtils.cpp
--- a/Engine/Source/Developer/BlueprintCompilerCppBackend/Private/BlueprintCompilerCppBackendUtils.cpp
+++ b/Engine/Source/Developer/BlueprintCompilerCppBackend/Private/BlueprintCompilerCppBackendUtils.cpp
@@ -95,7 +95,7 @@ FString FEmitHelper::HandleMetaData(const UField* Field, bool AddCategory, TArra
 #endif
 #define HANDLE_CPF_TAG(TagName, CheckedFlags) if (HasAllFlags(Flags, (CheckedFlags))) { Tags.Emplace(TagName); }
 
-TArray<FString> FEmitHelper::ProperyFlagsToTags(uint64 Flags)
+TArray<FString> FEmitHelper::ProperyFlagsToTags(uint64 Flags, bool bIsClassProperty)
 {
 	TArray<FString> Tags;
 
@@ -153,7 +153,11 @@ TArray<FString> FEmitHelper::ProperyFlagsToTags(uint64 Flags)
 	HANDLE_CPF_TAG(TEXT("Export"), CPF_ExportObject)
 	HANDLE_CPF_TAG(TEXT("NoClear"), CPF_NoClear)
 	HANDLE_CPF_TAG(TEXT("EditFixedSize"), CPF_EditFixedSize)
-	HANDLE_CPF_TAG(TEXT("NotReplicated"), CPF_RepSkip)
+	if (!bIsClassProperty)
+	{
+		HANDLE_CPF_TAG(TEXT("NotReplicated"), CPF_RepSkip)
+	}
+
 	HANDLE_CPF_TAG(TEXT("Interp"), CPF_Edit | CPF_BlueprintVisible | CPF_Interp)
 	HANDLE_CPF_TAG(TEXT("NonTransactional"), CPF_NonTransactional)
 	HANDLE_CPF_TAG(TEXT("BlueprintAssignable"), CPF_BlueprintAssignable)
diff --git a/Engine/Source/Developer/BlueprintCompilerCppBackend/Private/BlueprintCompilerCppBackendUtils.h b/Engine/Source/Developer/BlueprintCompilerCppBackend/Private/BlueprintCompilerCppBackendUtils.h
--- a/Engine/Source/Developer/BlueprintCompilerCppBackend/Private/BlueprintCompilerCppBackendUtils.h
+++ b/Engine/Source/Developer/BlueprintCompilerCppBackend/Private/BlueprintCompilerCppBackendUtils.h
@@ -224,7 +224,7 @@ struct FEmitHelper
 
 	static FString HandleMetaData(const UField* Field, bool AddCategory = true, TArray<FString>* AdditinalMetaData = nullptr);
 
-	static TArray<FString> ProperyFlagsToTags(uint64 Flags);
+	static TArray<FString> ProperyFlagsToTags(uint64 Flags, bool bIsClassProperty);
 
 	static TArray<FString> FunctionFlagsToTags(uint64 Flags);
 
@@ -295,6 +295,8 @@ struct FBackendHelperUMG
 {
 	static FString WidgetFunctionsInHeader(UClass* SourceClass);
 
+	static FString AdditionalHeaderIncludeForWidget(UClass* SourceClass);
+
 	// these function should use the same context as Constructor
 	static void CreateClassSubobjects(FEmitterLocalContext& Context);
 	static void EmitWidgetInitializationFunctions(FEmitterLocalContext& Context);
diff --git a/Engine/Source/Developer/BlueprintCompilerCppBackend/Private/BlueprintCompilerCppBackendValueHelper.cpp b/Engine/Source/Developer/BlueprintCompilerCppBackend/Private/BlueprintCompilerCppBackendValueHelper.cpp
--- a/Engine/Source/Developer/BlueprintCompilerCppBackend/Private/BlueprintCompilerCppBackendValueHelper.cpp
+++ b/Engine/Source/Developer/BlueprintCompilerCppBackend/Private/BlueprintCompilerCppBackendValueHelper.cpp
@@ -53,7 +53,7 @@ void FEmitDefaultValueHelper::OuterGenerate(FEmitterLocalContext& Context
 
 				FStringOutputDevice TypeDeclaration;
 				const uint32 CppTemplateTypeFlags = EPropertyExportCPPFlags::CPPF_CustomTypeName
-					| EPropertyExportCPPFlags::CPPF_NoConst | EPropertyExportCPPFlags::CPPF_NoRef
+					| EPropertyExportCPPFlags::CPPF_NoConst | EPropertyExportCPPFlags::CPPF_NoRef | EPropertyExportCPPFlags::CPPF_NoStaticArray
 					| EPropertyExportCPPFlags::CPPF_BlueprintCppBackend;
 				Property->ExportCppDeclaration(TypeDeclaration, EExportedDeclaration::Parameter, nullptr, CppTemplateTypeFlags, true);
 
diff --git a/Engine/Source/Developer/BlueprintCompilerCppBackend/Public/BlueprintCompilerCppBackendGatherDependencies.h b/Engine/Source/Developer/BlueprintCompilerCppBackend/Public/BlueprintCompilerCppBackendGatherDependencies.h
--- a/Engine/Source/Developer/BlueprintCompilerCppBackend/Public/BlueprintCompilerCppBackendGatherDependencies.h
+++ b/Engine/Source/Developer/BlueprintCompilerCppBackend/Public/BlueprintCompilerCppBackendGatherDependencies.h
@@ -32,6 +32,11 @@ public:
 
 	virtual void HandleObjectReference(UObject*& InObject, const UObject* InReferencingObject, const UProperty* InReferencingProperty) override;
 
+	UStruct* GetOriginalStruct() const
+	{
+		return OriginalStruct;
+	}
+
 public:
 	bool WillClassBeConverted(const UBlueprintGeneratedClass* InClass) const;
 
diff --git a/Engine/Source/Runtime/CoreUObject/Private/UObject/Property.cpp b/Engine/Source/Runtime/CoreUObject/Private/UObject/Property.cpp
--- a/Engine/Source/Runtime/CoreUObject/Private/UObject/Property.cpp
+++ b/Engine/Source/Runtime/CoreUObject/Private/UObject/Property.cpp
@@ -492,7 +492,8 @@ void UProperty::ExportCppDeclaration(FOutputDevice& Out, EExportedDeclaration::T
 	}
 
 	TCHAR ArrayStr[MAX_SPRINTF]=TEXT("");
-	if( ArrayDim != 1 )
+	const bool bExportStaticArray = 0 == (CPPF_NoStaticArray & AdditionalExportCPPFlags);
+	if ((ArrayDim != 1) && bExportStaticArray)
 	{
 		if (ArrayDimOverride)
 		{
diff --git a/Engine/Source/Runtime/CoreUObject/Public/UObject/UnrealType.h b/Engine/Source/Runtime/CoreUObject/Public/UObject/UnrealType.h
--- a/Engine/Source/Runtime/CoreUObject/Public/UObject/UnrealType.h
+++ b/Engine/Source/Runtime/CoreUObject/Public/UObject/UnrealType.h
@@ -31,6 +31,8 @@ enum EPropertyExportCPPFlags
 	CPPF_NoConst					=	0x00000010,
 	/** No reference '&' sign */
 	CPPF_NoRef						=	0x00000020,
+	/** No static array [%d] */
+	CPPF_NoStaticArray				=	0x00000040,
 	/** Blueprint compiler generated C++ code */
 	CPPF_BlueprintCppBackend		=	0x00000080,
 };
