UEBP-40 C++ code generation (WIP):
- Refactor, Introduced function FEmitHelper::GetCppName

diff --git a/Engine/Source/Developer/BlueprintCompilerCppBackend/Private/BlueprintCompilerCppBackend.cpp b/Engine/Source/Developer/BlueprintCompilerCppBackend/Private/BlueprintCompilerCppBackend.cpp
--- a/Engine/Source/Developer/BlueprintCompilerCppBackend/Private/BlueprintCompilerCppBackend.cpp
+++ b/Engine/Source/Developer/BlueprintCompilerCppBackend/Private/BlueprintCompilerCppBackend.cpp
@@ -431,9 +431,9 @@ FString FBlueprintCompilerCppBackend::EmitCallStatmentInner(FEmitterLocalContext
 	{
 		auto ContextInterfaceClass = CastChecked<UClass>(Statement.FunctionContext->Type.PinSubCategoryObject.Get());
 		ensure(ContextInterfaceClass->IsChildOf<UInterface>());
-		Result += FString::Printf(TEXT("I%s::Execute_%s(%s.GetObject(), ")
-			, *ContextInterfaceClass->GetName()
-			, *Statement.FunctionToCall->GetName()
+		Result += FString::Printf(TEXT("%s::Execute_%s(%s.GetObject(), ")
+			, *FEmitHelper::GetCppName(ContextInterfaceClass)
+			, *FEmitHelper::GetCppName(Statement.FunctionToCall)
 			, *TermToText(EmitterContext, Statement.FunctionContext, false));
 	}
 	else
@@ -444,7 +444,7 @@ FString FBlueprintCompilerCppBackend::EmitCallStatmentInner(FEmitterLocalContext
 				|| Statement.FunctionToCall->HasMetaData(TEXT("CustomStructureParam")) 
 				|| Statement.FunctionToCall->HasMetaData(TEXT("ArrayParm"));
 			auto OwnerClass = Statement.FunctionToCall->GetOuterUClass();
-			Result += bIsCustomThunk ? TEXT("FCustomThunkTemplates::") : FString::Printf(TEXT("%s%s::"), OwnerClass->GetPrefixCPP(), *OwnerClass->GetName());
+			Result += bIsCustomThunk ? TEXT("FCustomThunkTemplates::") : FString::Printf(TEXT("%s::"), *FEmitHelper::GetCppName(OwnerClass));
 		}
 		else if (bCallOnDifferentObject) //@TODO: Badness, could be a self reference wired to another instance!
 		{
@@ -455,7 +455,7 @@ FString FBlueprintCompilerCppBackend::EmitCallStatmentInner(FEmitterLocalContext
 		{
 			Result += TEXT("Super::");
 		}
-		Result += Statement.FunctionToCall->GetName();
+		Result += FEmitHelper::GetCppName(Statement.FunctionToCall);
 		if (Statement.bIsParentContext && FEmitHelper::ShouldHandleAsNativeEvent(Statement.FunctionToCall))
 		{
 			ensure(!bCallOnDifferentObject);
@@ -512,9 +512,8 @@ FString FBlueprintCompilerCppBackend::TermToText(FEmitterLocalContext& EmitterCo
 					: Cast<UClass>(Term->Context->Type.PinSubCategoryObject.Get());
 				if (MinimalClass)
 				{
-					ResultPath += FString::Printf(TEXT("GetDefaultValueSafe<%s%s>(")
-						, MinimalClass->GetPrefixCPP()
-						, *MinimalClass->GetName());
+					ResultPath += FString::Printf(TEXT("GetDefaultValueSafe<%s>(")
+						, *FEmitHelper::GetCppName(MinimalClass));
 				}
 				else
 				{
@@ -540,7 +539,7 @@ FString FBlueprintCompilerCppBackend::TermToText(FEmitterLocalContext& EmitterCo
 			Conditions = FSafeContextScopedEmmitter::ValidationChain(EmitterContext, Term->Context, *this);
 		}
 
-		ResultPath += Term->AssociatedVarProperty ? Term->AssociatedVarProperty->GetNameCPP() : Term->Name;
+		ResultPath += Term->AssociatedVarProperty ? FEmitHelper::GetCppName(Term->AssociatedVarProperty) : Term->Name;
 		return Conditions.IsEmpty()
 			? ResultPath
 			: FString::Printf(TEXT("((%s) ? (%s) : (%s))"), *Conditions, *ResultPath, *FEmitHelper::DefaultValue(EmitterContext, Term->Type));
diff --git a/Engine/Source/Developer/BlueprintCompilerCppBackend/Private/BlueprintCompilerCppBackendBase.cpp b/Engine/Source/Developer/BlueprintCompilerCppBackend/Private/BlueprintCompilerCppBackendBase.cpp
--- a/Engine/Source/Developer/BlueprintCompilerCppBackend/Private/BlueprintCompilerCppBackendBase.cpp
+++ b/Engine/Source/Developer/BlueprintCompilerCppBackend/Private/BlueprintCompilerCppBackendBase.cpp
@@ -74,7 +74,7 @@ void FBlueprintCompilerCppBackendBase::DeclareDelegates(UClass* SourceClass, TIn
 void FBlueprintCompilerCppBackendBase::GenerateCodeFromClass(UClass* SourceClass, TIndirectArray<FKismetFunctionContext>& Functions, bool bGenerateStubsOnly)
 {
 	auto CleanCppClassName = SourceClass->GetName();
-	CppClassName = FString(SourceClass->GetPrefixCPP()) + CleanCppClassName;
+	CppClassName = FEmitHelper::GetCppName(SourceClass);
 	
 	FGatherConvertedClassDependencies Dependencies(SourceClass);
 	EmitFileBeginning(CleanCppClassName, &Dependencies);
@@ -99,13 +99,13 @@ void FBlueprintCompilerCppBackendBase::GenerateCodeFromClass(UClass* SourceClass
 		Emit(Header, TEXT(")\n"));
 
 		UClass* SuperClass = SourceClass->GetSuperClass();
-		Emit(Header, *FString::Printf(TEXT("class %s : public %s%s"), *CppClassName, SuperClass->GetPrefixCPP(), *SuperClass->GetName()));
+		Emit(Header, *FString::Printf(TEXT("class %s : public %s"), *CppClassName, *FEmitHelper::GetCppName(SuperClass)));
 
 		for (auto& ImplementedInterface : SourceClass->Interfaces)
 		{
 			if (ImplementedInterface.Class)
 			{
-				Emit(Header, *FString::Printf(TEXT(", public I%s"), *ImplementedInterface.Class->GetName()));
+				Emit(Header, *FString::Printf(TEXT(", public %s"), *FEmitHelper::GetCppName(ImplementedInterface.Class)));
 			}
 		}
 	}
@@ -184,8 +184,7 @@ void FBlueprintCompilerCppBackendBase::ConstructFunction(FKismetFunctionContext&
 	TArray<UProperty*> LocalVariables;
 
 	{
-		FString FunctionName;
-		Function->GetName(FunctionName);
+		FString FunctionName = FEmitHelper::GetCppName(Function);
 
 		TArray<UProperty*> ArgumentList;
 
@@ -338,13 +337,12 @@ void FBlueprintCompilerCppBackendBase::ConstructFunction(FKismetFunctionContext&
 void FBlueprintCompilerCppBackendBase::GenerateCodeFromEnum(UUserDefinedEnum* SourceEnum)
 {
 	check(SourceEnum);
-	const FString Name = SourceEnum->GetName();
-	EmitFileBeginning(Name, nullptr);
+	EmitFileBeginning(SourceEnum->GetName(), nullptr);
 
 	Emit(Header, TEXT("UENUM(BlueprintType"));
 	EmitReplaceConvertedMetaData(SourceEnum);
 	Emit(Header, TEXT(")\nenum class "));
-	Emit(Header, *Name);
+	Emit(Header, *FEmitHelper::GetCppName(SourceEnum));
 	Emit(Header, TEXT(" : uint8\n{"));
 
 	for (int32 Index = 0; Index < SourceEnum->NumEnums(); ++Index)
@@ -378,7 +376,7 @@ void FBlueprintCompilerCppBackendBase::GenerateCodeFromStruct(UUserDefinedStruct
 	FGatherConvertedClassDependencies Dependencies(SourceStruct);
 	EmitFileBeginning(SourceStruct->GetName(), &Dependencies);
 
-	const FString NewName = FString(TEXT("F")) + SourceStruct->GetName();
+	const FString NewName = FEmitHelper::GetCppName(SourceStruct);
 	Emit(Header, TEXT("USTRUCT(BlueprintType"));
 	EmitReplaceConvertedMetaData(SourceStruct);
 	Emit(Header, TEXT(")\n"));
@@ -457,7 +455,7 @@ void FBlueprintCompilerCppBackendBase::EmitFileBeginning(const FString& CleanNam
 			{
 				if (auto ForwardDeclaredType = Cast<UClass>(Type))
 				{
-					Emit(Dst, *FString::Printf(TEXT("class %s;\n"), *(FString(ForwardDeclaredType->GetPrefixCPP()) + ForwardDeclaredType->GetName())));
+					Emit(Dst, *FString::Printf(TEXT("class %s;\n"), *FEmitHelper::GetCppName(ForwardDeclaredType)));
 				}
 			}
 
diff --git a/Engine/Source/Developer/BlueprintCompilerCppBackend/Private/BlueprintCompilerCppBackendUMG.cpp b/Engine/Source/Developer/BlueprintCompilerCppBackend/Private/BlueprintCompilerCppBackendUMG.cpp
--- a/Engine/Source/Developer/BlueprintCompilerCppBackend/Private/BlueprintCompilerCppBackendUMG.cpp
+++ b/Engine/Source/Developer/BlueprintCompilerCppBackend/Private/BlueprintCompilerCppBackendUMG.cpp
@@ -45,7 +45,7 @@ void FBackendHelperUMG::EmitWidgetInitializationFunctions(FEmitterLocalContext&
 {
 	if (auto WidgetClass = Cast<UWidgetBlueprintGeneratedClass>(Context.GetCurrentlyGeneratedClass()))
 	{
-		const FString CppClassName = FString(WidgetClass->GetPrefixCPP()) + WidgetClass->GetName();
+		const FString CppClassName = FEmitHelper::GetCppName(WidgetClass);
 
 		auto GenerateLocalProperty = [](FEmitterLocalContext& InContext, UProperty* InProperty, const uint8* DataPtr) -> FString
 		{
diff --git a/Engine/Source/Developer/BlueprintCompilerCppBackend/Private/BlueprintCompilerCppBackendUtils.cpp b/Engine/Source/Developer/BlueprintCompilerCppBackend/Private/BlueprintCompilerCppBackendUtils.cpp
--- a/Engine/Source/Developer/BlueprintCompilerCppBackend/Private/BlueprintCompilerCppBackendUtils.cpp
+++ b/Engine/Source/Developer/BlueprintCompilerCppBackend/Private/BlueprintCompilerCppBackendUtils.cpp
@@ -4,6 +4,104 @@
 #include "BlueprintCompilerCppBackendUtils.h"
 #include "EdGraphSchema_K2.h"
 
+FString FEmitterLocalContext::GenerateGetProperty(const UProperty* Property) const
+{
+	check(Property);
+	const UStruct* OwnerStruct = Property->GetOwnerStruct();
+	const FString OwnerName = FEmitHelper::GetCppName(OwnerStruct);
+	const FString OwnerPath = OwnerName + (OwnerStruct->IsA<UClass>() ? TEXT("::StaticClass()") : TEXT("::StaticStruct()"));
+	//return FString::Printf(TEXT("%s->FindPropertyByName(GET_MEMBER_NAME_CHECKED(%s, %s))"), *OwnerPath, *OwnerName, *Property->GetNameCPP());
+
+	return FString::Printf(TEXT("%s->FindPropertyByName(FName(TEXT(\"%s\")))"), *OwnerPath, *FEmitHelper::GetCppName(Property));
+}
+
+FString FEmitterLocalContext::GenerateUniqueLocalName()
+{
+	const FString UniqueNameBase = TEXT("__Local__");
+	const FString UniqueName = FString::Printf(TEXT("%s%d"), *UniqueNameBase, LocalNameIndexMax);
+	++LocalNameIndexMax;
+	return UniqueName;
+}
+
+FString FEmitterLocalContext::FindGloballyMappedObject(UObject* Object, bool bLoadIfNotFound)
+{
+	// TODO: check if not excluded
+
+	if (ActualClass && (Object == ActualClass))
+	{
+		return TEXT("GetClass()");
+	}
+
+	if (auto ObjClass = Cast<UClass>(Object))
+	{
+		auto BPGC = Cast<UBlueprintGeneratedClass>(ObjClass);
+		if (ObjClass->HasAnyClassFlags(CLASS_Native) || (BPGC && Dependencies.WillClassBeConverted(BPGC)))
+		{
+			return FString::Printf(TEXT("%s::StaticClass()"), *FEmitHelper::GetCppName(ObjClass));
+		}
+	}
+
+	// TODO Handle native structires, and special cases..
+
+	if (auto UDS = Cast<UScriptStruct>(Object))
+	{
+		// Check if  
+		// TODO: check if supported 
+		return FString::Printf(TEXT("%s::StaticStruct()"), *FEmitHelper::GetCppName(UDS));
+	}
+
+	if (auto UDE = Cast<UEnum>(Object))
+	{
+		// TODO:
+		// TODO: check if supported 
+		//return FString::Printf(TEXT("%s_StaticEnum()"), *UDE->GetName());
+		return FString::Printf(TEXT("FindObjectChecked<UEnum>(ANY_PACKAGE, TEXT(\"%s\"))"), *FEmitHelper::GetCppName(UDE));
+	}
+
+	int32 ObjectsCreatedPerClassIdx = INDEX_NONE;
+	if (ActualClass && Object && ObjectsCreatedPerClass.Find(Object, ObjectsCreatedPerClassIdx))
+	{
+		return FString::Printf(TEXT("CastChecked<%s>(%s::StaticClass()->ConvertedSubobjectsFromBPGC[%d])")
+			, *FEmitHelper::GetCppName(Object->GetClass())
+			, *FEmitHelper::GetCppName(ActualClass)
+			, ObjectsCreatedPerClassIdx);
+	}
+
+	// TODO: handle subobjects
+
+	if (bLoadIfNotFound && ensure(Object))
+	{
+		UClass* FoundClass = Object->GetClass();
+		const FString ClassString = FEmitHelper::GetCppName(FoundClass);
+		return FString::Printf(TEXT("LoadObject<%s>(nullptr, TEXT(\"%s\"))"), *ClassString, *(Object->GetPathName().ReplaceCharWithEscapedChar()));
+	}
+
+	return FString{};
+}
+
+FString FEmitHelper::GetCppName(const UField* Field)
+{
+	check(Field);
+	auto AsClass = Cast<UClass>(Field);
+	auto AsScriptStruct = Cast<UScriptStruct>(Field);
+	if (AsClass || AsScriptStruct)
+	{
+		auto AsClass = Cast<UClass>(Field);
+		if (AsClass && AsClass->HasAnyClassFlags(CLASS_Interface))
+		{
+			ensure(AsClass->IsChildOf<UInterface>());
+			return FString::Printf(TEXT("I%s"), *AsClass->GetName());
+		}
+		auto AsStruct = CastChecked<UStruct>(Field);
+		return FString::Printf(TEXT("%s%s"), AsStruct->GetPrefixCPP(), *AsStruct->GetName());
+	}
+	else if (auto AsProperty = Cast<UProperty>(Field))
+	{
+		return AsProperty->GetNameCPP();
+	}
+	return Field->GetName();
+}
+
 void FEmitHelper::ArrayToString(const TArray<FString>& Array, FString& OutString, const TCHAR* Separator)
 {
 	if (Array.Num())
@@ -357,7 +455,7 @@ FString FEmitHelper::EmitLifetimeReplicatedPropsImpl(UClass* SourceClass, const
 				Result += FString::Printf(TEXT("%s\tSuper::GetLifetimeReplicatedProps(OutLifetimeProps);\n"), InCurrentIndent);
 				bFunctionInitilzed = true;
 			}
-			Result += FString::Printf(TEXT("%s\tDOREPLIFETIME( %s, %s);\n"), InCurrentIndent, *CppClassName, *It->GetNameCPP());
+			Result += FString::Printf(TEXT("%s\tDOREPLIFETIME( %s, %s);\n"), InCurrentIndent, *CppClassName, *FEmitHelper::GetCppName(*It));
 		}
 	}
 	if (bFunctionInitilzed)
@@ -394,7 +492,7 @@ FString FEmitHelper::LiteralTerm(FEmitterLocalContext& EmitterContext, const FEd
 		auto TypeEnum = Cast<UEnum>(Type.PinSubCategoryObject.Get());
 		if (TypeEnum)
 		{
-			return FString::Printf(TEXT("%s::%s"), *TypeEnum->GetName(), CustomValue.IsEmpty() 
+			return FString::Printf(TEXT("%s::%s"), *FEmitHelper::GetCppName(TypeEnum), CustomValue.IsEmpty()
 				? *TypeEnum->GetEnumName(TypeEnum->NumEnums()-1)
 				: *CustomValue);
 		}
@@ -463,7 +561,7 @@ FString FEmitHelper::LiteralTerm(FEmitterLocalContext& EmitterContext, const FEd
 		{
 			//@todo:  This needs to be more robust, since import text isn't really proper for struct construction.
 			const bool bEmptyCustomValue = CustomValue.IsEmpty() || (CustomValue == TEXT("()"));
-			const FString StructName = FString(TEXT("F")) + StructType->GetName();
+			const FString StructName = *FEmitHelper::GetCppName(StructType);
 			if (bEmptyCustomValue)
 			{
 				return StructName + (StructType->IsA<UUserDefinedStruct>() ? TEXT("::GetDefaultValue()") : TEXT("{}"));
@@ -529,7 +627,7 @@ FString FEmitHelper::LiteralTerm(FEmitterLocalContext& EmitterContext, const FEd
 				return MappedObject;
 			}
 			UClass* FoundClass = Cast<UClass>(Type.PinSubCategoryObject.Get());
-			FString ClassString = FoundClass ? (FString(FoundClass->GetPrefixCPP()) + FoundClass->GetName()) : TEXT("UObject");
+			FString ClassString = FoundClass ? FEmitHelper::GetCppName(FoundClass) : TEXT("UObject");
 			return FString::Printf(TEXT("LoadObject<%s>(nullptr, TEXT(\"%s\"))"), *ClassString, *(LiteralObject->GetPathName().ReplaceCharWithEscapedChar()));
 		}
 		else
@@ -627,14 +725,14 @@ bool FEmitHelper::GenerateAutomaticCast(const FEdGraphPinType& LType, const FEdG
 		auto RTypeEnum = Cast<UEnum>(RType.PinSubCategoryObject.Get());
 		if (!RTypeEnum && LTypeEnum)
 		{
-			FString EnumCppType = !LTypeEnum->CppType.IsEmpty() ? LTypeEnum->CppType : LTypeEnum->GetName();
+			const FString EnumCppType = !LTypeEnum->CppType.IsEmpty() ? LTypeEnum->CppType : FEmitHelper::GetCppName(LTypeEnum);
 			OutCastBegin = FString::Printf(TEXT("static_cast<%s>("), *EnumCppType);
 			OutCastEnd = TEXT(")");
 			return true;
 		}
 		if (!LTypeEnum && RTypeEnum)
 		{
-			FString EnumCppType = !RTypeEnum->CppType.IsEmpty() ? RTypeEnum->CppType : RTypeEnum->GetName();
+			const FString EnumCppType = !RTypeEnum->CppType.IsEmpty() ? RTypeEnum->CppType : FEmitHelper::GetCppName(RTypeEnum);
 			OutCastBegin = FString::Printf(TEXT("EnumToByte<%s>("), *EnumCppType); 
 			OutCastEnd = TEXT(")");
 			return true;
@@ -648,7 +746,7 @@ bool FEmitHelper::GenerateAutomaticCast(const FEdGraphPinType& LType, const FEdG
 		UClass* RClass = Cast<UClass>(RType.PinSubCategoryObject.Get());
 		if (LClass && RClass && LClass->IsChildOf(RClass) && !RClass->IsChildOf(LClass))
 		{
-			OutCastBegin = FString::Printf(TEXT("CastChecked<%s%s>("), LClass->GetPrefixCPP(), *LClass->GetName());
+			OutCastBegin = FString::Printf(TEXT("CastChecked<%s>("), *FEmitHelper::GetCppName(LClass));
 			OutCastEnd = TEXT(")");
 			return true;
 		}
diff --git a/Engine/Source/Developer/BlueprintCompilerCppBackend/Private/BlueprintCompilerCppBackendUtils.h b/Engine/Source/Developer/BlueprintCompilerCppBackend/Private/BlueprintCompilerCppBackendUtils.h
--- a/Engine/Source/Developer/BlueprintCompilerCppBackend/Private/BlueprintCompilerCppBackendUtils.h
+++ b/Engine/Source/Developer/BlueprintCompilerCppBackend/Private/BlueprintCompilerCppBackendUtils.h
@@ -78,92 +78,23 @@ public:
 		return ObjectsCreatedPerClass;
 	}
 
-	// Universal functions
-	FString GenerateGetProperty(const UProperty* Property) const
-	{
-		check(Property);
-		const UStruct* OwnerStruct = Property->GetOwnerStruct();
-		const FString OwnerName = FString(OwnerStruct->GetPrefixCPP()) + OwnerStruct->GetName();
-		const FString OwnerPath = OwnerName + (OwnerStruct->IsA<UClass>() ? TEXT("::StaticClass()") : TEXT("::StaticStruct()"));
-		//return FString::Printf(TEXT("%s->FindPropertyByName(GET_MEMBER_NAME_CHECKED(%s, %s))"), *OwnerPath, *OwnerName, *Property->GetNameCPP());
+	// UNIVERSAL FUNCTIONS
 
-		return FString::Printf(TEXT("%s->FindPropertyByName(FName(TEXT(\"%s\")))"), *OwnerPath, *Property->GetNameCPP());
-	}
+	FString GenerateGetProperty(const UProperty* Property) const;
 
-	FString GenerateUniqueLocalName()
-	{
-		const FString UniqueNameBase = TEXT("__Local__");
-		const FString UniqueName = FString::Printf(TEXT("%s%d"), *UniqueNameBase, LocalNameIndexMax);
-		++LocalNameIndexMax;
-		return UniqueName;
-	}
+	FString GenerateUniqueLocalName();
 
 	UClass* GetCurrentlyGeneratedClass() const
 	{
 		return ActualClass;
 	}
 
-	/** All objects (that can be referenced from other package) that will have a different path in cooked build 
+	/** All objects (that can be referenced from other package) that will have a different path in cooked build
 	(due to the native code generation), should be handled by this function */
-	FString FindGloballyMappedObject(UObject* Object, bool bLoadIfNotFound = false)
-	{
-		// TODO: check if not excluded
-
-		if (ActualClass && (Object == ActualClass))
-		{
-			return TEXT("GetClass()");
-		}
-
-		if (auto ObjClass = Cast<UClass>(Object))
-		{
-			auto BPGC = Cast<UBlueprintGeneratedClass>(ObjClass);
-			if (ObjClass->HasAnyClassFlags(CLASS_Native) || (BPGC && Dependencies.WillClassBeConverted(BPGC)))
-			{
-				return FString::Printf(TEXT("%s%s::StaticClass()"), ObjClass->GetPrefixCPP(), *ObjClass->GetName());
-			}
-		}
-
-		// TODO Handle native structires, and special cases..
-
-		if (auto UDS = Cast<UScriptStruct>(Object))
-		{
-			// Check if  
-			// TODO: check if supported 
-			return FString::Printf(TEXT("%s%s::StaticStruct()"), UDS->GetPrefixCPP(), *UDS->GetName());
-		}
-
-		if (auto UDE = Cast<UEnum>(Object))
-		{
-			// TODO:
-			// TODO: check if supported 
-			//return FString::Printf(TEXT("%s_StaticEnum()"), *UDE->GetName());
-			return FString::Printf(TEXT("FindObjectChecked<UEnum>(ANY_PACKAGE, TEXT(\"%s\"))"), *UDE->GetName());
-		}
-
-		int32 ObjectsCreatedPerClassIdx = INDEX_NONE;
-		if (ActualClass && Object && ObjectsCreatedPerClass.Find(Object, ObjectsCreatedPerClassIdx))
-		{
-			return FString::Printf(TEXT("CastChecked<%s%s>(%s%s::StaticClass()->ConvertedSubobjectsFromBPGC[%d])")
-				, Object->GetClass()->GetPrefixCPP()
-				, *Object->GetClass()->GetName()
-				, ActualClass->GetPrefixCPP()
-				, *ActualClass->GetName()
-				, ObjectsCreatedPerClassIdx);
-		}
-
-		// TODO: handle subobjects
-
-		if (bLoadIfNotFound && ensure(Object))
-		{
-			UClass* FoundClass = Object->GetClass();
-			const FString ClassString = FString(FoundClass->GetPrefixCPP()) + FoundClass->GetName();
-			return FString::Printf(TEXT("LoadObject<%s>(nullptr, TEXT(\"%s\"))"), *ClassString, *(Object->GetPathName().ReplaceCharWithEscapedChar()));
-		}
-
-		return FString{};
-	}
+	FString FindGloballyMappedObject(UObject* Object, bool bLoadIfNotFound = false);
 
 	// TEXT
+
 	void IncreaseIndent()
 	{
 		Indent += TEXT("\t");
@@ -214,6 +145,8 @@ public:
 
 struct FEmitHelper
 {
+	static FString GetCppName(const UField* Field);
+
 	static void ArrayToString(const TArray<FString>& Array, FString& OutString, const TCHAR* Separator);
 
 	static bool HasAllFlags(uint64 Flags, uint64 FlagsToCheck);
diff --git a/Engine/Source/Developer/BlueprintCompilerCppBackend/Private/BlueprintCompilerCppBackendValueHelper.cpp b/Engine/Source/Developer/BlueprintCompilerCppBackend/Private/BlueprintCompilerCppBackendValueHelper.cpp
--- a/Engine/Source/Developer/BlueprintCompilerCppBackend/Private/BlueprintCompilerCppBackendValueHelper.cpp
+++ b/Engine/Source/Developer/BlueprintCompilerCppBackend/Private/BlueprintCompilerCppBackendValueHelper.cpp
@@ -69,7 +69,7 @@ void FEmitDefaultValueHelper::OuterGenerate(FEmitterLocalContext& Context
 					: ((EPropertyAccessOperator::Pointer == AccessOperator) ? TEXT("->") : TEXT("."));
 				const bool bStaticArray = (Property->ArrayDim > 1);
 				const FString ArrayPost = bStaticArray ? FString::Printf(TEXT("[%d]"), ArrayIndex) : TEXT("");
-				PathToMember = FString::Printf(TEXT("%s%s%s%s"), *OuterPath, *AccessOperatorStr, *Property->GetNameCPP(), *ArrayPost);
+				PathToMember = FString::Printf(TEXT("%s%s%s%s"), *OuterPath, *AccessOperatorStr, *FEmitHelper::GetCppName(Property), *ArrayPost);
 			}
 			const uint8* ValuePtr = Property->ContainerPtrToValuePtr<uint8>(DataContainer, ArrayIndex);
 			const uint8* DefaultValuePtr = OptionalDefaultDataContainer ? Property->ContainerPtrToValuePtr<uint8>(OptionalDefaultDataContainer, ArrayIndex) : nullptr;
@@ -81,7 +81,7 @@ void FEmitDefaultValueHelper::OuterGenerate(FEmitterLocalContext& Context
 FString FEmitDefaultValueHelper::GenerateGetDefaultValue(const UUserDefinedStruct* Struct, const FGatherConvertedClassDependencies& Dependencies)
 {
 	check(Struct);
-	const FString StructName = FString(TEXT("F")) + Struct->GetName();
+	const FString StructName = FEmitHelper::GetCppName(Struct);
 	FString Result = FString::Printf(TEXT("\tstatic %s GetDefaultValue()\n\t{\n\t\t%s DefaultData__;\n"), *StructName, *StructName);
 
 	FStructOnScope StructData(Struct);
@@ -114,7 +114,7 @@ void FEmitDefaultValueHelper::InnerGenerate(FEmitterLocalContext& Context, const
 				check(StructProperty->Struct);
 				if (bGenerateEmptyStructConstructor)
 				{
-					ValueStr = FString::Printf(TEXT("F%s{}"), *StructProperty->Struct->GetName()); //don;t override existing values
+					ValueStr = FString::Printf(TEXT("%s{}"), *FEmitHelper::GetCppName(StructProperty->Struct)); //don;t override existing values
 				}
 				bComplete = false;
 			}
@@ -282,9 +282,11 @@ FString FEmitDefaultValueHelper::HandleNonNativeComponent(FEmitterLocalContext&
 			}
 			else
 			{
-				Context.AddHighPriorityLine(FString::Printf(TEXT("%s%s = CreateDefaultSubobject<%s%s>(TEXT(\"%s\"));")
-					, VariableProperty == nullptr ? TEXT("auto ") : TEXT(""), *VariableName
-					, ComponentClass->GetPrefixCPP(), *ComponentClass->GetName(), *ComponentTemplate->GetName()));
+				Context.AddHighPriorityLine(FString::Printf(TEXT("%s%s = CreateDefaultSubobject<%s>(TEXT(\"%s\"));")
+					, (VariableProperty == nullptr) ? TEXT("auto ") : TEXT("")
+					, *VariableName
+					, *FEmitHelper::GetCppName(ComponentClass)
+					, *ComponentTemplate->GetName()));
 
 				Context.AddLine(FString::Printf(TEXT("%s->CreationMethod = EComponentCreationMethod::SimpleConstructionScript;"), *VariableName));
 
@@ -333,9 +335,8 @@ struct FDependenciesHelper
 		}
 		for (auto LocStruct : GatherDependencies.ConvertedClasses)
 		{
-			Context.AddLine(FString::Printf(TEXT("GetClass()->ConvertedSubobjectsFromBPGC.Add(%s%s::StaticClass());")
-				, LocStruct->GetPrefixCPP()
-				, *LocStruct->GetName()));
+			Context.AddLine(FString::Printf(TEXT("GetClass()->ConvertedSubobjectsFromBPGC.Add(%s::StaticClass());")
+				, *FEmitHelper::GetCppName(LocStruct)));
 		}
 
 		if (GatherDependencies.ConvertedStructs.Num())
@@ -344,9 +345,8 @@ struct FDependenciesHelper
 		}
 		for (auto LocStruct : GatherDependencies.ConvertedStructs)
 		{
-			Context.AddLine(FString::Printf(TEXT("GetClass()->ConvertedSubobjectsFromBPGC.Add(%s%s::StaticStruct());")
-				, LocStruct->GetPrefixCPP()
-				, *LocStruct->GetName()));
+			Context.AddLine(FString::Printf(TEXT("GetClass()->ConvertedSubobjectsFromBPGC.Add(%s::StaticStruct());")
+				, *FEmitHelper::GetCppName(LocStruct)));
 		}
 
 		if (GatherDependencies.Assets.Num())
@@ -363,19 +363,20 @@ struct FDependenciesHelper
 	static void AddStaticFunctionsForDependencies(FEmitterLocalContext& Context, const FGatherConvertedClassDependencies& GatherDependencies)
 	{
 		auto OriginalClass = Context.GetCurrentlyGeneratedClass();
-		const FString CppClassName = FString::Printf(TEXT("%s%s"), OriginalClass->GetPrefixCPP(), *OriginalClass->GetName());
+		const FString CppClassName = FEmitHelper::GetCppName(OriginalClass);
 		// __StaticDependenciesConvertedClasses
 		Context.AddLine(FString::Printf(TEXT("void %s::__StaticDependenciesConvertedClasses(TArray<FName>& OutNames)"), *CppClassName));
 		Context.AddLine(TEXT("{"));
 		Context.IncreaseIndent();
 
+		//TODO: the names must be valid.
 		for (auto LocClass : GatherDependencies.ConvertedClasses)
 		{
 			Context.AddLine(FString::Printf(TEXT("OutNames.Add(TEXT(\"%s\"));"), *LocClass->GetName()));
 		}
-		for (auto LocClass : GatherDependencies.ConvertedStructs)
+		for (auto LocStruct : GatherDependencies.ConvertedStructs)
 		{
-			Context.AddLine(FString::Printf(TEXT("OutNames.Add(TEXT(\"%s\"));"), *LocClass->GetName()));
+			Context.AddLine(FString::Printf(TEXT("OutNames.Add(TEXT(\"%s\"));"), *LocStruct->GetName()));
 		}
 
 		Context.DecreaseIndent();
@@ -424,7 +425,7 @@ struct FDependenciesHelper
 FString FEmitDefaultValueHelper::GenerateConstructor(UClass* InBPGC, const FGatherConvertedClassDependencies& Dependencies)
 {
 	auto BPGC = CastChecked<UBlueprintGeneratedClass>(InBPGC);
-	const FString CppClassName = FString(BPGC->GetPrefixCPP()) + BPGC->GetName();
+	const FString CppClassName = FEmitHelper::GetCppName(BPGC);
 
 	FEmitterLocalContext Context(BPGC, Dependencies);
 	Context.AddLine(FString::Printf(TEXT("%s::%s(const FObjectInitializer& ObjectInitializer) : Super(ObjectInitializer)"), *CppClassName, *CppClassName));
@@ -502,10 +503,9 @@ FString FEmitDefaultValueHelper::GenerateConstructor(UClass* InBPGC, const FGath
 			FString NativeName;
 			const bool bFound = Context.FindLocalObject_InConstructor(Obj, NativeName);
 			ensure(bFound);
-			Context.AddHighPriorityLine(FString::Printf(TEXT("auto %s = CastChecked<%s%s>(%s::StaticClass()->ConvertedSubobjectsFromBPGC[%d]);")
+			Context.AddHighPriorityLine(FString::Printf(TEXT("auto %s = CastChecked<%s>(%s::StaticClass()->ConvertedSubobjectsFromBPGC[%d]);")
 				, *NativeName
-				, Obj->GetClass()->GetPrefixCPP()
-				, *Obj->GetClass()->GetName()
+				, *FEmitHelper::GetCppName(Obj->GetClass())
 				, *CppClassName
 				, SubobjectIndex));
 			SubobjectIndex++;
@@ -627,9 +627,9 @@ FString FEmitDefaultValueHelper::HandleClassSubobject(FEmitterLocalContext& Cont
 	Context.AddHighPriorityLine(FString::Printf(
 		TEXT("auto %s = NewObject<%s%s>(%s, TEXT(\"%s\"));")
 		, *LocalNativeName
-		, ObjectClass->GetPrefixCPP()
-		, *ObjectClass->GetName()
-		, *OuterStr, *Object->GetName()));
+		, *FEmitHelper::GetCppName(ObjectClass)
+		, *OuterStr
+		, *Object->GetName()));
 	if (AddAsSubobjectOfClass)
 	{
 		ensure(Context.bCreatingObjectsPerClass);
@@ -683,8 +683,8 @@ FString FEmitDefaultValueHelper::HandleInstancedSubobject(FEmitterLocalContext&
 	auto CDO = BPGC ? BPGC->GetDefaultObject(false) : nullptr;
 	if (!bIsEditorOnlySubobject && ensure(CDO) && (CDO == Object->GetOuter()))
 	{
-		Context.AddHighPriorityLine(FString::Printf(TEXT("auto %s = CastChecked<%s%s>(GetDefaultSubobjectByName(TEXT(\"%s\")));")
-			, *LocalNativeName, ObjectClass->GetPrefixCPP(), *ObjectClass->GetName(), *Object->GetName()));
+		Context.AddHighPriorityLine(FString::Printf(TEXT("auto %s = CastChecked<%s>(GetDefaultSubobjectByName(TEXT(\"%s\")));")
+			, *LocalNativeName, *FEmitHelper::GetCppName(ObjectClass), *Object->GetName()));
 
 		const UObject* ObjectArchetype = Object->GetArchetype();
 		for (auto Property : TFieldRange<const UProperty>(ObjectClass))
@@ -702,9 +702,11 @@ FString FEmitDefaultValueHelper::HandleInstancedSubobject(FEmitterLocalContext&
 		{
 			return FString();
 		}
-		Context.AddHighPriorityLine(FString::Printf(TEXT("auto %s = NewObject<%s%s>(%s, TEXT(\"%s\"));")
-			, *LocalNativeName, ObjectClass->GetPrefixCPP(), *ObjectClass->GetName()
-			, *OuterStr, *Object->GetName()));
+		Context.AddHighPriorityLine(FString::Printf(TEXT("auto %s = NewObject<%s>(%s, TEXT(\"%s\"));")
+			, *LocalNativeName
+			, *FEmitHelper::GetCppName(ObjectClass)
+			, *OuterStr
+			, *Object->GetName()));
 	}
 
 	return LocalNativeName;
diff --git a/Engine/Source/Developer/BlueprintNativeCodeGen/Private/NativeCodeGenerationTool.cpp b/Engine/Source/Developer/BlueprintNativeCodeGen/Private/NativeCodeGenerationTool.cpp
--- a/Engine/Source/Developer/BlueprintNativeCodeGen/Private/NativeCodeGenerationTool.cpp
+++ b/Engine/Source/Developer/BlueprintNativeCodeGen/Private/NativeCodeGenerationTool.cpp
@@ -43,21 +43,20 @@ struct FGeneratedCodeData
 		for (auto Iter : ClassDependencies.ConvertedClasses)
 		{
 			DependentObjects.Add(Iter);
-				}
+		}
 		for (auto Iter : ClassDependencies.ConvertedStructs)
-			{
+		{
 			DependentObjects.Add(Iter);
-			}
+		}
 		for (auto Iter : ClassDependencies.ConvertedEnum)
-			{
+		{
 			DependentObjects.Add(Iter);
 		}
 
 		TypeDependencies.Empty();
 		for (auto Obj : DependentObjects)
-				{
-					TypeDependencies += Obj->GetPathName();
-					TypeDependencies += TEXT("\n");
+		{
+			TypeDependencies += FString::Printf(TEXT("%s \t%s\n"), *Obj->GetClass()->GetName(), *Obj->GetPathName());
 		}
 		DependentObjects.Add(InBlueprint.GeneratedClass);
 
