FConvertedBlueprintsDependencies - next iteration.

diff --git a/Engine/Source/Developer/BlueprintCompilerCppBackend/Private/BlueprintCompilerCppBackendBase.cpp b/Engine/Source/Developer/BlueprintCompilerCppBackend/Private/BlueprintCompilerCppBackendBase.cpp
--- a/Engine/Source/Developer/BlueprintCompilerCppBackend/Private/BlueprintCompilerCppBackendBase.cpp
+++ b/Engine/Source/Developer/BlueprintCompilerCppBackend/Private/BlueprintCompilerCppBackendBase.cpp
@@ -144,8 +144,7 @@ void FBlueprintCompilerCppBackendBase::GenerateCodeFromClass(UClass* SourceClass
 		}
 	}
 
-	Emit(Header, TEXT("\n\tstatic void __StaticDependenciesConvertedClasses(TArray<FName>& OutNames);\n"));
-	Emit(Header, TEXT("\n\tstatic void __StaticDependenciesAssets(TArray<FName>& OutNames);\n"));
+	Emit(Header, TEXT("\n\tstatic void __StaticDependenciesAssets(TArray<FName>& OutPackagePaths);\n"));
 
 	Emit(Header, *FBackendHelperUMG::WidgetFunctionsInHeader(SourceClass));
 
diff --git a/Engine/Source/Developer/BlueprintCompilerCppBackend/Private/BlueprintCompilerCppBackendValueHelper.cpp b/Engine/Source/Developer/BlueprintCompilerCppBackend/Private/BlueprintCompilerCppBackendValueHelper.cpp
--- a/Engine/Source/Developer/BlueprintCompilerCppBackend/Private/BlueprintCompilerCppBackendValueHelper.cpp
+++ b/Engine/Source/Developer/BlueprintCompilerCppBackend/Private/BlueprintCompilerCppBackendValueHelper.cpp
@@ -364,23 +364,6 @@ struct FDependenciesHelper
 	{
 		auto OriginalClass = Context.GetCurrentlyGeneratedClass();
 		const FString CppClassName = FEmitHelper::GetCppName(OriginalClass);
-		// __StaticDependenciesConvertedClasses
-		Context.AddLine(FString::Printf(TEXT("void %s::__StaticDependenciesConvertedClasses(TArray<FName>& OutNames)"), *CppClassName));
-		Context.AddLine(TEXT("{"));
-		Context.IncreaseIndent();
-
-		//TODO: the names must be valid.
-		for (auto LocClass : GatherDependencies.ConvertedClasses)
-		{
-			Context.AddLine(FString::Printf(TEXT("OutNames.Add(TEXT(\"%s\"));"), *LocClass->GetName()));
-		}
-		for (auto LocStruct : GatherDependencies.ConvertedStructs)
-		{
-			Context.AddLine(FString::Printf(TEXT("OutNames.Add(TEXT(\"%s\"));"), *LocStruct->GetName()));
-		}
-
-		Context.DecreaseIndent();
-		Context.AddLine(TEXT("}"));
 
 		// __StaticDependenciesAssets
 		Context.AddLine(FString::Printf(TEXT("void %s::__StaticDependenciesAssets(TArray<FName>& OutPackagePaths)"), *CppClassName));
@@ -406,7 +389,7 @@ struct FDependenciesHelper
 		Context.AddLine(TEXT("{"));
 		Context.IncreaseIndent();
 		Context.AddLine(FString::Printf(
-			TEXT("FConvertedBlueprintsDependencies::Get().RegisterClass(TEXT(\"%s\"), &%s::__StaticDependenciesConvertedClasses, &%s::__StaticDependenciesAssets);")
+			TEXT("FConvertedBlueprintsDependencies::Get().RegisterClass(TEXT(\"%s\"), &%s::__StaticDependenciesAssets);")
 			, *OriginalClass->GetName()
 			, *CppClassName
 			, *CppClassName));
diff --git a/Engine/Source/Runtime/CoreUObject/Private/Blueprint/BlueprintSupport.cpp b/Engine/Source/Runtime/CoreUObject/Private/Blueprint/BlueprintSupport.cpp
--- a/Engine/Source/Runtime/CoreUObject/Private/Blueprint/BlueprintSupport.cpp
+++ b/Engine/Source/Runtime/CoreUObject/Private/Blueprint/BlueprintSupport.cpp
@@ -1730,32 +1730,18 @@ FConvertedBlueprintsDependencies& FConvertedBlueprintsDependencies::Get()
 	return ConvertedBlueprintsDependencies;
 }
 
-void FConvertedBlueprintsDependencies::RegisterClass(FName ClassName, GetDependenciesNamesFunc GetConvertedClasses, GetDependenciesNamesFunc GetAssets)
+void FConvertedBlueprintsDependencies::RegisterClass(FName ClassName, GetDependenciesNamesFunc GetAssets)
 {
 	check(!ClassNameToGetter.Contains(ClassName));
-
-	FGetters Getters;
-	Getters.GetAssets = GetAssets;
-	Getters.GetConvertedClasses = GetConvertedClasses;
-	ClassNameToGetter.Add(ClassName, Getters);
-}
-
-void FConvertedBlueprintsDependencies::GetConvertedClasses(FName ClassName, TArray<FName>& OutConvertedClassNames) const
-{
-	auto GettersPtr = ClassNameToGetter.Find(ClassName);
-	auto Func = (GettersPtr) ? GettersPtr->GetConvertedClasses : nullptr;
-	ensure(Func || !GettersPtr);
-	if (Func)
-	{
-		Func(OutConvertedClassNames);
-	}
+	ensure(GetAssets);
+	ClassNameToGetter.Add(ClassName, GetAssets);
 }
 
 void FConvertedBlueprintsDependencies::GetAssets(FName ClassName, TArray<FName>& OutPackagePaths) const
 {
-	auto GettersPtr = ClassNameToGetter.Find(ClassName);
-	auto Func = (GettersPtr) ? GettersPtr->GetAssets : nullptr;
-	ensure(Func || !GettersPtr);
+	auto FuncPtr = ClassNameToGetter.Find(ClassName);
+	auto Func = (FuncPtr) ? (*FuncPtr) : nullptr;
+	ensure(Func || !FuncPtr);
 	if (Func)
 	{
 		Func(OutPackagePaths);
diff --git a/Engine/Source/Runtime/CoreUObject/Public/Blueprint/BlueprintSupport.h b/Engine/Source/Runtime/CoreUObject/Public/Blueprint/BlueprintSupport.h
--- a/Engine/Source/Runtime/CoreUObject/Public/Blueprint/BlueprintSupport.h
+++ b/Engine/Source/Runtime/CoreUObject/Public/Blueprint/BlueprintSupport.h
@@ -165,18 +165,7 @@ struct COREUOBJECT_API FConvertedBlueprintsDependencies
 
 private:
 
-	struct FGetters
-	{
-		GetDependenciesNamesFunc GetConvertedClasses;
-		GetDependenciesNamesFunc GetAssets;
-
-		FGetters()
-			: GetConvertedClasses(nullptr)
-			, GetAssets(nullptr)
-		{}
-	};
-
-	TMap<FName, FGetters> ClassNameToGetter;
+	TMap<FName, GetDependenciesNamesFunc> ClassNameToGetter;
 
 public:
 
@@ -187,10 +176,7 @@ public:
 		return ClassNameToGetter.Num() > 0;
 	}
 
-	void RegisterClass(FName ClassName, GetDependenciesNamesFunc GetConvertedClasses, GetDependenciesNamesFunc GetAssets);
-
-	/** Get all converted classes (and structures) names necessary for the class with the given class name. */
-	void GetConvertedClasses(FName ClassName, TArray<FName>& OutConvertedClassNames) const;
+	void RegisterClass(FName ClassName, GetDependenciesNamesFunc GetAssets);
 
 	/** Get all assets paths necessary for the class with the given class name and all converted classes that dependencies. */
 	void GetAssets(FName ClassName, TArray<FName>& OutPackagePaths) const;
