Fix shadowed variable

diff --git a/Engine/Source/Developer/BlueprintCompilerCppBackend/Private/BlueprintCompilerCppBackendUtils.cpp b/Engine/Source/Developer/BlueprintCompilerCppBackend/Private/BlueprintCompilerCppBackendUtils.cpp
--- a/Engine/Source/Developer/BlueprintCompilerCppBackend/Private/BlueprintCompilerCppBackendUtils.cpp
+++ b/Engine/Source/Developer/BlueprintCompilerCppBackend/Private/BlueprintCompilerCppBackendUtils.cpp
@@ -86,7 +86,6 @@ FString FEmitHelper::GetCppName(const UField* Field)
 	auto AsScriptStruct = Cast<UScriptStruct>(Field);
 	if (AsClass || AsScriptStruct)
 	{
-		auto AsClass = Cast<UClass>(Field);
 		if (AsClass && AsClass->HasAnyClassFlags(CLASS_Interface))
 		{
 			ensure(AsClass->IsChildOf<UInterface>());
