[UEBP-40] BP->C++ code generation (WiP) - Some fixes to generated code for default subobject instancing.

change summary:
- modified FEmitDefaultValueHelper::HandleInstancedSubobject() to allow callers to control whether or not to generate code to create the instance or access it (if it should already have been created).
- modified FEmitDefaultValueHelper::GenerateConstructor() to generate code to assign a fallback native root component if available (to match the runtime side in the non-generated code case)

diff --git a/Engine/Source/Developer/BlueprintCompilerCppBackend/Private/BlueprintCompilerCppBackendUtils.h b/Engine/Source/Developer/BlueprintCompilerCppBackend/Private/BlueprintCompilerCppBackendUtils.h
--- a/Engine/Source/Developer/BlueprintCompilerCppBackend/Private/BlueprintCompilerCppBackendUtils.h
+++ b/Engine/Source/Developer/BlueprintCompilerCppBackend/Private/BlueprintCompilerCppBackendUtils.h
@@ -221,7 +221,7 @@ private:
 
 	static FString HandleNonNativeComponent(FEmitterLocalContext& Context, const USCS_Node* Node, TSet<const UProperty*>& OutHandledProperties, const USCS_Node* ParentNode = nullptr);
 
-	static FString HandleInstancedSubobject(FEmitterLocalContext& Context, UObject* Object, bool bSkipEditorOnlyCheck = false);
+	static FString HandleInstancedSubobject(FEmitterLocalContext& Context, UObject* Object, bool bCreateInstance = true, bool bSkipEditorOnlyCheck = false);
 };
 
 struct FBackendHelperUMG
diff --git a/Engine/Source/Developer/BlueprintCompilerCppBackend/Private/BlueprintCompilerCppBackendValueHelper.cpp b/Engine/Source/Developer/BlueprintCompilerCppBackend/Private/BlueprintCompilerCppBackendValueHelper.cpp
--- a/Engine/Source/Developer/BlueprintCompilerCppBackend/Private/BlueprintCompilerCppBackendValueHelper.cpp
+++ b/Engine/Source/Developer/BlueprintCompilerCppBackend/Private/BlueprintCompilerCppBackendValueHelper.cpp
@@ -209,7 +209,7 @@ FString FEmitDefaultValueHelper::HandleSpecialTypes(FEmitterLocalContext& Contex
 
 			if (!Context.bCreatingObjectsPerClass && Property->HasAnyPropertyFlags(CPF_InstancedReference))
 			{
-				const FString CreateAsInstancedSubobject = HandleInstancedSubobject(Context, Object);
+				const FString CreateAsInstancedSubobject = HandleInstancedSubobject(Context, Object, Object->HasAnyFlags(RF_ArchetypeObject));
 				if (!CreateAsInstancedSubobject.IsEmpty())
 				{
 					return CreateAsInstancedSubobject;
@@ -502,6 +502,7 @@ FString FEmitDefaultValueHelper::GenerateConstructor(UClass* InBPGC, const FGath
 	check(CDO && ParentCDO);
 	Context.AddLine(TEXT(""));
 
+	FString NativeRootComponentFallback;
 	TSet<const UProperty*> HandledProperties;
 
 	// Generate ctor init code for native class default subobjects that are always instanced (e.g. components).
@@ -522,7 +523,17 @@ FString FEmitDefaultValueHelper::GenerateConstructor(UClass* InBPGC, const FGath
 			// Skip ctor code gen for editor-only subobjects, since they won't be used by the runtime. Any dependencies on editor-only subobjects will be handled later (see HandleInstancedSubobject).
 			if (!bIsEditorOnlySubobject)
 			{
-				HandleInstancedSubobject(Context, DSO, true);
+				const FString VariableName = HandleInstancedSubobject(Context, DSO, false, true);
+
+				// Keep track of which component can be used as a root, in case it's not explicitly set.
+				if (NativeRootComponentFallback.IsEmpty())
+				{
+					USceneComponent* SceneComponent = Cast<USceneComponent>(DSO);
+					if (SceneComponent && !SceneComponent->AttachParent && SceneComponent->CreationMethod == EComponentCreationMethod::Native)
+					{
+						NativeRootComponentFallback = VariableName;
+					}
+				}
 			}
 		}
 	}
@@ -537,6 +548,11 @@ FString FEmitDefaultValueHelper::GenerateConstructor(UClass* InBPGC, const FGath
 		{
 			HandledProperties.Add(RootComponentProperty);
 		}
+		else if (!NativeRootComponentFallback.IsEmpty())
+		{
+			Context.AddLine(FString::Printf(TEXT("RootComponent = %s;"), *NativeRootComponentFallback));
+			HandledProperties.Add(RootComponentProperty);
+		}
 		else
 		{
 			bNeedsRootComponentAssignment = true;
@@ -630,7 +646,7 @@ FString FEmitDefaultValueHelper::HandleClassSubobject(FEmitterLocalContext& Cont
 	return LocalNativeName;
 }
 
-FString FEmitDefaultValueHelper::HandleInstancedSubobject(FEmitterLocalContext& Context, UObject* Object, bool bSkipEditorOnlyCheck)
+FString FEmitDefaultValueHelper::HandleInstancedSubobject(FEmitterLocalContext& Context, UObject* Object, bool bCreateInstance, bool bSkipEditorOnlyCheck)
 {
 	check(Object);
 
@@ -666,8 +682,16 @@ FString FEmitDefaultValueHelper::HandleInstancedSubobject(FEmitterLocalContext&
 	auto CDO = BPGC ? BPGC->GetDefaultObject(false) : nullptr;
 	if (!bIsEditorOnlySubobject && ensure(CDO) && (CDO == Object->GetOuter()))
 	{
-		Context.AddHighPriorityLine(FString::Printf(TEXT("auto %s = CastChecked<%s>(GetDefaultSubobjectByName(TEXT(\"%s\")));")
-			, *LocalNativeName, *FEmitHelper::GetCppName(ObjectClass), *Object->GetName()));
+		if (bCreateInstance)
+		{
+			Context.AddHighPriorityLine(FString::Printf(TEXT("auto %s = CreateDefaultSubobject<%s>(TEXT(\"%s\"));")
+				, *LocalNativeName, *FEmitHelper::GetCppName(ObjectClass), *Object->GetName()));
+		}
+		else
+		{
+			Context.AddHighPriorityLine(FString::Printf(TEXT("auto %s = CastChecked<%s>(GetDefaultSubobjectByName(TEXT(\"%s\")));")
+				, *LocalNativeName, *FEmitHelper::GetCppName(ObjectClass), *Object->GetName()));
+		}
 
 		const UObject* ObjectArchetype = Object->GetArchetype();
 		for (auto Property : TFieldRange<const UProperty>(ObjectClass))
