UEBP-40 C++ code generation (WIP):
- refactor
- DynamicClass support
- various bugs

diff --git a/Engine/Source/Developer/BlueprintCompilerCppBackend/Private/BlueprintCompilerCppBackend.cpp b/Engine/Source/Developer/BlueprintCompilerCppBackend/Private/BlueprintCompilerCppBackend.cpp
--- a/Engine/Source/Developer/BlueprintCompilerCppBackend/Private/BlueprintCompilerCppBackend.cpp
+++ b/Engine/Source/Developer/BlueprintCompilerCppBackend/Private/BlueprintCompilerCppBackend.cpp
@@ -575,9 +575,8 @@ FString FBlueprintCompilerCppBackend::LatentFunctionInfoTermToText(FEmitterLocal
 	return FEmitHelper::LiteralTerm(EmitterContext, Term->Type, StructValues, nullptr);
 }
 
-FString FBlueprintCompilerCppBackend::InnerFunctionImplementation(FKismetFunctionContext& FunctionContext, const FGatherConvertedClassDependencies& Dependencies, bool bUseSwitchState)
+FString FBlueprintCompilerCppBackend::InnerFunctionImplementation(FKismetFunctionContext& FunctionContext, FEmitterLocalContext& EmitterContext, bool bUseSwitchState)
 {
-	FEmitterLocalContext EmitterContext(FunctionContext.NewClass, Dependencies);
 	EmitterContext.IncreaseIndent();
 	if (bUseSwitchState)
 	{
@@ -733,5 +732,8 @@ FString FBlueprintCompilerCppBackend::InnerFunctionImplementation(FKismetFunctio
 	}
 
 
-	return EmitterContext.GetResult();
+	const FString Result = EmitterContext.GetResult();
+	EmitterContext.ClearResult();
+	EmitterContext.DecreaseIndent();
+	return Result;
 }
diff --git a/Engine/Source/Developer/BlueprintCompilerCppBackend/Private/BlueprintCompilerCppBackend.h b/Engine/Source/Developer/BlueprintCompilerCppBackend/Private/BlueprintCompilerCppBackend.h
--- a/Engine/Source/Developer/BlueprintCompilerCppBackend/Private/BlueprintCompilerCppBackend.h
+++ b/Engine/Source/Developer/BlueprintCompilerCppBackend/Private/BlueprintCompilerCppBackend.h
@@ -9,7 +9,7 @@ struct FEmitterLocalContext;
 class FBlueprintCompilerCppBackend : public FBlueprintCompilerCppBackendBase
 {
 protected:
-	virtual FString InnerFunctionImplementation(FKismetFunctionContext& FunctionContext, const FGatherConvertedClassDependencies& Dependencies, bool bUseSwitchState) override;
+	virtual FString InnerFunctionImplementation(FKismetFunctionContext& FunctionContext, FEmitterLocalContext& EmitterContext, bool bUseSwitchState) override;
 
 protected:
 	void EmitCallStatment(FEmitterLocalContext& EmitterContext, FKismetFunctionContext& FunctionContext, FBlueprintCompiledStatement& Statement);
diff --git a/Engine/Source/Developer/BlueprintCompilerCppBackend/Private/BlueprintCompilerCppBackendBase.cpp b/Engine/Source/Developer/BlueprintCompilerCppBackend/Private/BlueprintCompilerCppBackendBase.cpp
--- a/Engine/Source/Developer/BlueprintCompilerCppBackend/Private/BlueprintCompilerCppBackendBase.cpp
+++ b/Engine/Source/Developer/BlueprintCompilerCppBackend/Private/BlueprintCompilerCppBackendBase.cpp
@@ -130,17 +130,19 @@ void FBlueprintCompilerCppBackendBase::GenerateCodeFromClass(UClass* SourceClass
 		Emit(Header, TEXT("\n"));
 	}
 
+	FEmitterLocalContext EmitterContext(SourceClass, Dependencies);
+
 	if (!bIsInterface)
 	{
 		Emit(Header, *FString::Printf(TEXT("\t%s(const FObjectInitializer& ObjectInitializer = FObjectInitializer::Get());\n\n"), *CppClassName));
-		Emit(Body, *FEmitDefaultValueHelper::GenerateConstructor(SourceClass, Dependencies));
+		Emit(Body, *FEmitDefaultValueHelper::GenerateConstructor(EmitterContext));
 	}
 
 	for (int32 i = 0; i < Functions.Num(); ++i)
 	{
 		if (Functions[i].IsValid())
 		{
-			ConstructFunction(Functions[i], Dependencies, bGenerateStubsOnly);
+			ConstructFunction(Functions[i], EmitterContext, bGenerateStubsOnly);
 		}
 	}
 
@@ -170,7 +172,7 @@ void FBlueprintCompilerCppBackendBase::DeclareLocalVariables(FKismetFunctionCont
 	}
 }
 
-void FBlueprintCompilerCppBackendBase::ConstructFunction(FKismetFunctionContext& FunctionContext, const FGatherConvertedClassDependencies& Dependencies, bool bGenerateStubOnly)
+void FBlueprintCompilerCppBackendBase::ConstructFunction(FKismetFunctionContext& FunctionContext, FEmitterLocalContext& EmitterContext, bool bGenerateStubOnly)
 {
 	if (FunctionContext.IsDelegateSignature())
 	{
@@ -324,11 +326,11 @@ void FBlueprintCompilerCppBackendBase::ConstructFunction(FKismetFunctionContext&
 			}
 		}
 
-		const FString FunctionImplementation = InnerFunctionImplementation(FunctionContext, Dependencies, bUseSwitchState);
+		const FString FunctionImplementation = InnerFunctionImplementation(FunctionContext, EmitterContext, bUseSwitchState);
 		Emit(Body, *FunctionImplementation);
 	}
 
-	const FString ReturnValueString = ReturnValue ? (FString(TEXT(" ")) + ReturnValue->GetName()) : TEXT("");
+	const FString ReturnValueString = ReturnValue ? (FString(TEXT(" ")) + FEmitHelper::GetCppName(ReturnValue)) : TEXT("");
 	Emit(Body, *FString::Printf(TEXT("\treturn%s;\n"), *ReturnValueString));
 	Emit(Body, TEXT("}\n\n"));
 }
@@ -400,7 +402,7 @@ void FBlueprintCompilerCppBackendBase::EmitFileBeginning(const FString& CleanNam
 
 	if (Dependencies)
 	{
-		Emit(Body, *FBackendHelperUMG::AdditionalHeaderIncludeForWidget(Cast<UClass>(Dependencies->GetOriginalStruct())));
+		Emit(Header, *FBackendHelperUMG::AdditionalHeaderIncludeForWidget(Cast<UClass>(Dependencies->GetOriginalStruct())));
 
 		TSet<FString> AlreadyIncluded;
 		AlreadyIncluded.Add(CleanName);
diff --git a/Engine/Source/Developer/BlueprintCompilerCppBackend/Private/BlueprintCompilerCppBackendBase.h b/Engine/Source/Developer/BlueprintCompilerCppBackend/Private/BlueprintCompilerCppBackendBase.h
--- a/Engine/Source/Developer/BlueprintCompilerCppBackend/Private/BlueprintCompilerCppBackendBase.h
+++ b/Engine/Source/Developer/BlueprintCompilerCppBackend/Private/BlueprintCompilerCppBackendBase.h
@@ -4,6 +4,8 @@
 #include "IBlueprintCompilerCppBackendModule.h"
 #include "BlueprintCompilerCppBackendGatherDependencies.h"
 
+struct FEmitterLocalContext;
+
 /** The class generates all native code except of function body (notice InnerFunctionImplementation in not implemented) */
 class FBlueprintCompilerCppBackendBase : public IBlueprintCompilerCppBackend
 {
@@ -70,7 +72,7 @@ protected:
 			}
 
 	*/
-	virtual FString InnerFunctionImplementation(FKismetFunctionContext& FunctionContext, const FGatherConvertedClassDependencies& Dependencies, bool bUseSwitchState) PURE_VIRTUAL(FBlueprintCompilerCppBackendBase::InnerFunctionImplementation, return FString(););
+	virtual FString InnerFunctionImplementation(FKismetFunctionContext& FunctionContext, FEmitterLocalContext& EmitterContext, bool bUseSwitchState) PURE_VIRTUAL(FBlueprintCompilerCppBackendBase::InnerFunctionImplementation, return FString(););
 
 	void EmitStructProperties(FStringOutputDevice& Target, UStruct* SourceClass);
 
@@ -78,7 +80,7 @@ protected:
 	void DeclareLocalVariables(FKismetFunctionContext& FunctionContext, TArray<UProperty*>& LocalVariables);
 	
 	/** Builds both the header declaration and body implementation of a function */
-	void ConstructFunction(FKismetFunctionContext& FunctionContext, const FGatherConvertedClassDependencies& Dependencies, bool bGenerateStubOnly);
+	void ConstructFunction(FKismetFunctionContext& FunctionContext, FEmitterLocalContext& EmitterContext, bool bGenerateStubOnly);
 
 	void EmitFileBeginning(const FString& CleanName, FGatherConvertedClassDependencies* Dependencies);
 
diff --git a/Engine/Source/Developer/BlueprintCompilerCppBackend/Private/BlueprintCompilerCppBackendUMG.cpp b/Engine/Source/Developer/BlueprintCompilerCppBackend/Private/BlueprintCompilerCppBackendUMG.cpp
--- a/Engine/Source/Developer/BlueprintCompilerCppBackend/Private/BlueprintCompilerCppBackendUMG.cpp
+++ b/Engine/Source/Developer/BlueprintCompilerCppBackend/Private/BlueprintCompilerCppBackendUMG.cpp
@@ -31,12 +31,12 @@ void FBackendHelperUMG::CreateClassSubobjects(FEmitterLocalContext& Context)
 		if (WidgetClass->WidgetTree)
 		{
 			ensure(WidgetClass->WidgetTree->GetOuter() == Context.GetCurrentlyGeneratedClass());
-			FEmitDefaultValueHelper::HandleClassSubobject(Context, WidgetClass->WidgetTree);
+			FEmitDefaultValueHelper::HandleClassSubobject(Context, WidgetClass->WidgetTree, FEmitterLocalContext::EClassSubobjectList::MiscConvertedSubobjects);
 		}
 		for (auto Anim : WidgetClass->Animations)
 		{
 			ensure(Anim->GetOuter() == Context.GetCurrentlyGeneratedClass());
-			FEmitDefaultValueHelper::HandleClassSubobject(Context, Anim);
+			FEmitDefaultValueHelper::HandleClassSubobject(Context, Anim, FEmitterLocalContext::EClassSubobjectList::MiscConvertedSubobjects);
 		}
 	}
 }
diff --git a/Engine/Source/Developer/BlueprintCompilerCppBackend/Private/BlueprintCompilerCppBackendUtils.cpp b/Engine/Source/Developer/BlueprintCompilerCppBackend/Private/BlueprintCompilerCppBackendUtils.cpp
--- a/Engine/Source/Developer/BlueprintCompilerCppBackend/Private/BlueprintCompilerCppBackendUtils.cpp
+++ b/Engine/Source/Developer/BlueprintCompilerCppBackend/Private/BlueprintCompilerCppBackendUtils.cpp
@@ -25,7 +25,56 @@ FString FEmitterLocalContext::GenerateUniqueLocalName()
 
 FString FEmitterLocalContext::FindGloballyMappedObject(UObject* Object, bool bLoadIfNotFound)
 {
-	// TODO: check if not excluded
+	if (ActualClass && Object && Object->IsIn(ActualClass))
+	{
+		if (const FString* NamePtr = (CurrentCodeType == EGeneratedCodeType::SubobjectsOfClass) ? ClassSubobjectsMap.Find(Object) : nullptr)
+		{
+			return *NamePtr;
+		}
+
+		if (const FString* NamePtr = (CurrentCodeType == EGeneratedCodeType::CommonConstructor) ? CommonSubobjectsMap.Find(Object) : nullptr)
+		{
+			return *NamePtr;
+		}
+
+		int32 ObjectsCreatedPerClassIdx = INDEX_NONE;
+		if (MiscConvertedSubobjects.Find(Object, ObjectsCreatedPerClassIdx))
+		{
+			return FString::Printf(TEXT("CastChecked<%s>(CastChecked<UDynamicClass>(%s::StaticClass())->MiscConvertedSubobjects[%d])")
+				, *FEmitHelper::GetCppName(Object->GetClass())
+				, *FEmitHelper::GetCppName(ActualClass)
+				, ObjectsCreatedPerClassIdx);
+		}
+		if (DynamicBindingObjects.Find(Object, ObjectsCreatedPerClassIdx))
+		{
+			return FString::Printf(TEXT("CastChecked<%s>(CastChecked<UDynamicClass>(%s::StaticClass())->DynamicBindingObjects[%d])")
+				, *FEmitHelper::GetCppName(Object->GetClass())
+				, *FEmitHelper::GetCppName(ActualClass)
+				, ObjectsCreatedPerClassIdx);
+		}
+		if (ComponentTemplates.Find(Object, ObjectsCreatedPerClassIdx))
+		{
+			return FString::Printf(TEXT("CastChecked<%s>(CastChecked<UDynamicClass>(%s::StaticClass())->ComponentTemplates[%d])")
+				, *FEmitHelper::GetCppName(Object->GetClass())
+				, *FEmitHelper::GetCppName(ActualClass)
+				, ObjectsCreatedPerClassIdx);
+		}
+		if (Timelines.Find(Object, ObjectsCreatedPerClassIdx))
+		{
+			return FString::Printf(TEXT("CastChecked<%s>(CastChecked<UDynamicClass>(%s::StaticClass())->Timelines[%d])")
+				, *FEmitHelper::GetCppName(Object->GetClass())
+				, *FEmitHelper::GetCppName(ActualClass)
+				, ObjectsCreatedPerClassIdx);
+		}
+
+		if ((CurrentCodeType == EGeneratedCodeType::SubobjectsOfClass) || (CurrentCodeType == EGeneratedCodeType::CommonConstructor))
+		{
+			if (Object == ActualClass->GetDefaultObject(false))
+			{
+				return TEXT("this");
+			}
+		}
+	}
 
 	if (ActualClass && (Object == ActualClass))
 	{
@@ -58,15 +107,6 @@ FString FEmitterLocalContext::FindGloballyMappedObject(UObject* Object, bool bLo
 		return FString::Printf(TEXT("FindObjectChecked<UEnum>(ANY_PACKAGE, TEXT(\"%s\"))"), *FEmitHelper::GetCppName(UDE));
 	}
 
-	int32 ObjectsCreatedPerClassIdx = INDEX_NONE;
-	if (ActualClass && Object && ObjectsCreatedPerClass.Find(Object, ObjectsCreatedPerClassIdx))
-	{
-		return FString::Printf(TEXT("CastChecked<%s>(%s::StaticClass()->ConvertedSubobjectsFromBPGC[%d])")
-			, *FEmitHelper::GetCppName(Object->GetClass())
-			, *FEmitHelper::GetCppName(ActualClass)
-			, ObjectsCreatedPerClassIdx);
-	}
-
 	// TODO: handle subobjects
 
 	if (bLoadIfNotFound && ensure(Object))
@@ -106,7 +146,8 @@ FString FEmitHelper::GetCppName(const UField* Field)
 		}
 		return AsProperty->GetNameCPP();
 	}
-	if (!Field->HasAnyFlags(RF_Native))
+
+	if (!Field->HasAnyFlags(RF_Native) && !Field->IsA<UUserDefinedEnum>())
 	{
 		return ::UnicodeToCPPIdentifier(Field->GetName(), false, TEXT("bpf__"));
 	}
diff --git a/Engine/Source/Developer/BlueprintCompilerCppBackend/Private/BlueprintCompilerCppBackendUtils.h b/Engine/Source/Developer/BlueprintCompilerCppBackend/Private/BlueprintCompilerCppBackendUtils.h
--- a/Engine/Source/Developer/BlueprintCompilerCppBackend/Private/BlueprintCompilerCppBackendUtils.h
+++ b/Engine/Source/Developer/BlueprintCompilerCppBackend/Private/BlueprintCompilerCppBackendUtils.h
@@ -5,7 +5,25 @@
 
 struct FEmitterLocalContext
 {
+	enum class EClassSubobjectList
+	{
+		ComponentTemplates,
+		Timelines,
+		DynamicBindingObjects,
+		MiscConvertedSubobjects,
+	};
+
+	enum class EGeneratedCodeType
+	{
+		SubobjectsOfClass,
+		CommonConstructor,
+		Regular,
+	};
+
+	EGeneratedCodeType CurrentCodeType;
+
 private:
+	// TEXT
 	FString Indent;
 	FString Result;
 	int32 LocalNameIndexMax;
@@ -13,69 +31,81 @@ private:
 	TArray<FString> LowPriorityLines;
 	UClass* ActualClass;
 
-	//ConstructorOnly
-	TMap<UObject*, FString> NativeObjectNamesInConstructor;
-	TArray<UObject*> ObjectsCreatedPerClass;
+	//ConstructorOnly Local Names
+	TMap<UObject*, FString> ClassSubobjectsMap;
+	//ConstructorOnly Local Names
+	TMap<UObject*, FString> CommonSubobjectsMap;
 
-	const FGatherConvertedClassDependencies& Dependencies;
+	TArray<UObject*> MiscConvertedSubobjects;
+	TArray<UObject*> DynamicBindingObjects;
+	TArray<UObject*> ComponentTemplates;
+	TArray<UObject*> Timelines;
 
 public:
 
-	bool bCreatingObjectsPerClass;
-	bool bInsideConstructor;
+	const FGatherConvertedClassDependencies& Dependencies;
 
 	FEmitterLocalContext(UClass* InClass, const FGatherConvertedClassDependencies& InDependencies)
-		: LocalNameIndexMax(0)
+		: CurrentCodeType(EGeneratedCodeType::Regular)
+		, LocalNameIndexMax(0)
 		, ActualClass(InClass)
 		, Dependencies(InDependencies)
-		, bCreatingObjectsPerClass(false)
-		, bInsideConstructor(false)
 	{}
 
-	// Functions to use in constructor only
-	bool FindLocalObject_InConstructor(UObject* Object, FString& OutNamePath) const
+	UClass* GetActualClass() const
 	{
-		ensure(bInsideConstructor);
-		if (Object == ActualClass)
+		return ActualClass;
+	}
+
+	// CONSTRUCTOR FUNCTIONS
+
+	const TCHAR* ClassSubobjectListName(EClassSubobjectList ListType)
 		{
-			OutNamePath = TEXT("GetClass()");
-			return true;
+		if (ListType == EClassSubobjectList::ComponentTemplates)
+		{
+			return TEXT("ComponentTemplates");
 		}
-
-		if (const FString* NamePtr = NativeObjectNamesInConstructor.Find(Object))
+		if (ListType == EClassSubobjectList::Timelines)
+		{
+			return TEXT("Timelines");
+		}
+		if (ListType == EClassSubobjectList::DynamicBindingObjects)
 		{
-			OutNamePath = *NamePtr;
-			return true;
+			return TEXT("DynamicBindingObjects");
+		}
+		return TEXT("MiscConvertedSubobjects");
 		}
 
-		return false;
+	void RegisterClassSubobject(UObject* Object, EClassSubobjectList ListType)
+	{
+		if (ListType == EClassSubobjectList::ComponentTemplates)
+		{
+			ComponentTemplates.Add(Object);
 	}
-
-	void AddObjectFromLocalProperty_InConstructor(UObject* Object, const FString& NativeName)
+		else if (ListType == EClassSubobjectList::Timelines)
 	{
-		ensure(bInsideConstructor);
-		NativeObjectNamesInConstructor.Add(Object, NativeName);
+			Timelines.Add(Object);
 	}
-
-	FString AddNewObject_InConstructor(UObject* Object, bool bAddToObjectsCreatedPerClass)
+		else if(ListType == EClassSubobjectList::DynamicBindingObjects)
 	{
-		ensure(bInsideConstructor);
-		const FString UniqueName = GenerateUniqueLocalName();
-		NativeObjectNamesInConstructor.Add(Object, UniqueName);
-		if (bAddToObjectsCreatedPerClass)
+			DynamicBindingObjects.Add(Object);
+		}
+		else
 		{
-			ensure(bCreatingObjectsPerClass);
-			check(!ObjectsCreatedPerClass.Contains(Object));
-			ObjectsCreatedPerClass.Add(Object);
+			MiscConvertedSubobjects.Add(Object);
+		}
 		}
 
-		return UniqueName;
+	void AddClassSubObject_InConstructor(UObject* Object, const FString& NativeName)
+	{
+		ensure(CurrentCodeType == EGeneratedCodeType::SubobjectsOfClass);
+		ClassSubobjectsMap.Add(Object, NativeName);
 	}
 
-	const TArray<UObject*>& GetObjectsCreatedPerClass_InConstructor()
+	void AddCommonSubObject_InConstructor(UObject* Object, const FString& NativeName)
 	{
-		ensure(bInsideConstructor);
-		return ObjectsCreatedPerClass;
+		ensure(CurrentCodeType == EGeneratedCodeType::CommonConstructor);
+		ClassSubobjectsMap.Add(Object, NativeName);
 	}
 
 	// UNIVERSAL FUNCTIONS
@@ -93,7 +123,7 @@ public:
 	(due to the native code generation), should be handled by this function */
 	FString FindGloballyMappedObject(UObject* Object, bool bLoadIfNotFound = false);
 
-	// TEXT
+	// TEXT FUNCTIONS
 
 	void IncreaseIndent()
 	{
@@ -194,7 +224,7 @@ struct FEmitDefaultValueHelper
 {
 	static FString GenerateGetDefaultValue(const UUserDefinedStruct* Struct, const FGatherConvertedClassDependencies& Dependencies);
 
-	static FString GenerateConstructor(UClass* BPGC, const FGatherConvertedClassDependencies& Dependencies);
+	static FString GenerateConstructor(FEmitterLocalContext& Context);
 
 	enum class EPropertyAccessOperator
 	{
@@ -212,7 +242,7 @@ struct FEmitDefaultValueHelper
 
 	// Creates the subobject (of class) returns it's native local name, 
 	// returns empty string if cannot handle
-	static FString HandleClassSubobject(FEmitterLocalContext& Context, UObject* Object);
+	static FString HandleClassSubobject(FEmitterLocalContext& Context, UObject* Object, FEmitterLocalContext::EClassSubobjectList ListOfSubobjectsTyp);
 
 private:
 	// Returns native term, 
diff --git a/Engine/Source/Developer/BlueprintCompilerCppBackend/Private/BlueprintCompilerCppBackendValueHelper.cpp b/Engine/Source/Developer/BlueprintCompilerCppBackend/Private/BlueprintCompilerCppBackendValueHelper.cpp
--- a/Engine/Source/Developer/BlueprintCompilerCppBackend/Private/BlueprintCompilerCppBackendValueHelper.cpp
+++ b/Engine/Source/Developer/BlueprintCompilerCppBackend/Private/BlueprintCompilerCppBackendValueHelper.cpp
@@ -190,24 +190,17 @@ FString FEmitDefaultValueHelper::HandleSpecialTypes(FEmitterLocalContext& Contex
 				}
 			}
 
-			{
-				FString LocalMappedObject;
-				if (Context.bInsideConstructor && Context.FindLocalObject_InConstructor(Object, LocalMappedObject))
-				{
-					return LocalMappedObject;
-				}
-			}
-
+			const bool bCreatingSubObjectsOfClass = (Context.CurrentCodeType == FEmitterLocalContext::EGeneratedCodeType::SubobjectsOfClass);
 			{
 				auto BPGC = Context.GetCurrentlyGeneratedClass();
 				auto CDO = BPGC ? BPGC->GetDefaultObject(false) : nullptr;
-				if (BPGC && Object && CDO && Object->IsIn(BPGC) && !Object->IsIn(CDO) && Context.bCreatingObjectsPerClass)
+				if (BPGC && Object && CDO && Object->IsIn(BPGC) && !Object->IsIn(CDO) && bCreatingSubObjectsOfClass)
 				{
-					return HandleClassSubobject(Context, Object);
+					return HandleClassSubobject(Context, Object, FEmitterLocalContext::EClassSubobjectList::MiscConvertedSubobjects);
 				}
 			}
 
-			if (!Context.bCreatingObjectsPerClass && Property->HasAnyPropertyFlags(CPF_InstancedReference))
+			if (!bCreatingSubObjectsOfClass && Property->HasAnyPropertyFlags(CPF_InstancedReference))
 			{
 				const FString CreateAsInstancedSubobject = HandleInstancedSubobject(Context, Object, Object->HasAnyFlags(RF_ArchetypeObject));
 				if (!CreateAsInstancedSubobject.IsEmpty())
@@ -255,19 +248,26 @@ FString FEmitDefaultValueHelper::HandleSpecialTypes(FEmitterLocalContext& Contex
 FString FEmitDefaultValueHelper::HandleNonNativeComponent(FEmitterLocalContext& Context, const USCS_Node* Node, TSet<const UProperty*>& OutHandledProperties, const USCS_Node* ParentNode)
 {
 	check(Node);
+	check(Context.CurrentCodeType == FEmitterLocalContext::EGeneratedCodeType::CommonConstructor);
 
-	FString VariableName;
+	FString NativeVariablePropertyName;
 	UBlueprintGeneratedClass* BPGC = CastChecked<UBlueprintGeneratedClass>(Context.GetCurrentlyGeneratedClass());
 	if (UActorComponent* ComponentTemplate = Node->GetActualComponentTemplate(BPGC))
 	{
-		VariableName = Node->VariableName.ToString();
-		Context.AddObjectFromLocalProperty_InConstructor(ComponentTemplate, VariableName);
+		const FString VariableCleanName = Node->VariableName.ToString();
 
-		const UObjectProperty* VariableProperty = FindField<UObjectProperty>(BPGC, *VariableName);
+		const UObjectProperty* VariableProperty = FindField<UObjectProperty>(BPGC, *VariableCleanName);
 		if (VariableProperty)
 		{
+			NativeVariablePropertyName = FEmitHelper::GetCppName(VariableProperty);
 			OutHandledProperties.Add(VariableProperty);
 		}
+		else
+		{
+			NativeVariablePropertyName = VariableCleanName;
+		}
+
+		Context.AddCommonSubObject_InConstructor(ComponentTemplate, NativeVariablePropertyName);
 
 		if (ComponentTemplate->GetOuter() == BPGC)
 		{
@@ -284,31 +284,33 @@ FString FEmitDefaultValueHelper::HandleNonNativeComponent(FEmitterLocalContext&
 			{
 				Context.AddHighPriorityLine(FString::Printf(TEXT("%s%s = CreateDefaultSubobject<%s>(TEXT(\"%s\"));")
 					, (VariableProperty == nullptr) ? TEXT("auto ") : TEXT("")
-					, *VariableName
+					, *NativeVariablePropertyName
 					, *FEmitHelper::GetCppName(ComponentClass)
 					, *ComponentTemplate->GetName()));
 
-				Context.AddLine(FString::Printf(TEXT("%s->CreationMethod = EComponentCreationMethod::SimpleConstructionScript;"), *VariableName));
+				Context.AddLine(FString::Printf(TEXT("%s->CreationMethod = EComponentCreationMethod::SimpleConstructionScript;"), *NativeVariablePropertyName));
 
 				FString ParentVariableName;
 				if (ParentNode)
 				{
-					ParentVariableName = ParentNode->VariableName.ToString();
+					const FString CleanParentVariableName = ParentNode->VariableName.ToString();
+					const UObjectProperty* ParentVariableProperty = FindField<UObjectProperty>(BPGC, *CleanParentVariableName);
+					ParentVariableName = ParentVariableProperty ? FEmitHelper::GetCppName(ParentVariableProperty) : CleanParentVariableName;
 				}
 				else if (USceneComponent* ParentComponentTemplate = Node->GetParentComponentTemplate(CastChecked<UBlueprint>(BPGC->ClassGeneratedBy)))
 				{
-					Context.FindLocalObject_InConstructor(ParentComponentTemplate, ParentVariableName);
+					ParentVariableName = Context.FindGloballyMappedObject(ParentComponentTemplate);
 				}
 
 				if (!ParentVariableName.IsEmpty())
 				{
-					Context.AddLine(FString::Printf(TEXT("%s->AttachParent = %s;"), *VariableName, *ParentVariableName));
+					Context.AddLine(FString::Printf(TEXT("%s->AttachParent = %s;"), *NativeVariablePropertyName, *ParentVariableName));
 				}
 			}
 
 			for (auto Property : TFieldRange<const UProperty>(ComponentClass))
 			{
-				OuterGenerate(Context, Property, VariableName
+				OuterGenerate(Context, Property, NativeVariablePropertyName
 					, reinterpret_cast<const uint8*>(ComponentTemplate)
 					, reinterpret_cast<const uint8*>(ObjectToCompare)
 					, EPropertyAccessOperator::Pointer);
@@ -322,45 +324,45 @@ FString FEmitDefaultValueHelper::HandleNonNativeComponent(FEmitterLocalContext&
 		HandleNonNativeComponent(Context, ChildNode, OutHandledProperties, Node);
 	}
 
-	return VariableName;
+	return NativeVariablePropertyName;
 }
 
 struct FDependenciesHelper
 {
-	static void AddDependenciesInConstructor(FEmitterLocalContext& Context, const FGatherConvertedClassDependencies& GatherDependencies)
+	static void AddDependenciesInConstructor(FEmitterLocalContext& Context)
 	{
-		if (GatherDependencies.ConvertedClasses.Num())
+		if (Context.Dependencies.ConvertedClasses.Num())
 		{
 			Context.AddLine(TEXT("// List of all referenced converted classes"));
 		}
-		for (auto LocStruct : GatherDependencies.ConvertedClasses)
+		for (auto LocStruct : Context.Dependencies.ConvertedClasses)
 		{
-			Context.AddLine(FString::Printf(TEXT("GetClass()->ConvertedSubobjectsFromBPGC.Add(%s::StaticClass());")
+			Context.AddLine(FString::Printf(TEXT("CastChecked<UDynamicClass>(GetClass())->ReferencedConvertedFields.Add(%s::StaticClass());")
 				, *FEmitHelper::GetCppName(LocStruct)));
 		}
 
-		if (GatherDependencies.ConvertedStructs.Num())
+		if (Context.Dependencies.ConvertedStructs.Num())
 		{
 			Context.AddLine(TEXT("// List of all referenced converted structures"));
 		}
-		for (auto LocStruct : GatherDependencies.ConvertedStructs)
+		for (auto LocStruct : Context.Dependencies.ConvertedStructs)
 		{
-			Context.AddLine(FString::Printf(TEXT("GetClass()->ConvertedSubobjectsFromBPGC.Add(%s::StaticStruct());")
+			Context.AddLine(FString::Printf(TEXT("CastChecked<UDynamicClass>(GetClass())->ReferencedConvertedFields.Add(%s::StaticStruct());")
 				, *FEmitHelper::GetCppName(LocStruct)));
 		}
 
-		if (GatherDependencies.Assets.Num())
+		if (Context.Dependencies.Assets.Num())
 		{
 			Context.AddLine(TEXT("// List of all referenced assets"));
 		}
-		for (auto LocAsset : GatherDependencies.Assets)
+		for (auto LocAsset : Context.Dependencies.Assets)
 		{
 			const FString AssetStr = Context.FindGloballyMappedObject(LocAsset, true);
-			Context.AddLine(FString::Printf(TEXT("GetClass()->ConvertedSubobjectsFromBPGC.Add(%s);"), *AssetStr));
+			Context.AddLine(FString::Printf(TEXT("CastChecked<UDynamicClass>(GetClass())->UsedAssets.Add(%s);"), *AssetStr));
 		}
 	}
 
-	static void AddStaticFunctionsForDependencies(FEmitterLocalContext& Context, const FGatherConvertedClassDependencies& GatherDependencies)
+	static void AddStaticFunctionsForDependencies(FEmitterLocalContext& Context)
 	{
 		auto OriginalClass = Context.GetCurrentlyGeneratedClass();
 		const FString CppClassName = FEmitHelper::GetCppName(OriginalClass);
@@ -370,7 +372,7 @@ struct FDependenciesHelper
 		Context.AddLine(TEXT("{"));
 		Context.IncreaseIndent();
 
-		for (UObject* LocAsset : GatherDependencies.Assets)
+		for (UObject* LocAsset : Context.Dependencies.Assets)
 		{
 			const FString PakagePath = LocAsset->GetOutermost()->GetPathName();
 			Context.AddLine(FString::Printf(TEXT("OutPackagePaths.Add(TEXT(\"%s\"));"), *PakagePath));
@@ -405,20 +407,18 @@ struct FDependenciesHelper
 	}
 };
 
-FString FEmitDefaultValueHelper::GenerateConstructor(UClass* InBPGC, const FGatherConvertedClassDependencies& Dependencies)
+FString FEmitDefaultValueHelper::GenerateConstructor(FEmitterLocalContext& Context)
 {
-	auto BPGC = CastChecked<UBlueprintGeneratedClass>(InBPGC);
+	auto BPGC = CastChecked<UBlueprintGeneratedClass>(Context.GetActualClass());
 	const FString CppClassName = FEmitHelper::GetCppName(BPGC);
 
-	FEmitterLocalContext Context(BPGC, Dependencies);
 	Context.AddLine(FString::Printf(TEXT("%s::%s(const FObjectInitializer& ObjectInitializer) : Super(ObjectInitializer)"), *CppClassName, *CppClassName));
 	Context.AddLine(TEXT("{"));
 	Context.IncreaseIndent();
-	Context.bInsideConstructor = true;
 
 	// When CDO is created create all subobjects owned by the class
-	TArray<UActorComponent*> ActorComponentTempatesOwnedByClass = BPGC->ComponentTemplates;
 	{
+		TArray<UActorComponent*> ActorComponentTempatesOwnedByClass = BPGC->ComponentTemplates;
 		// Gather all CT from SCS and IH, the remaining ones are generated for class..
 		if (auto SCS = BPGC->SimpleConstructionScript)
 		{
@@ -440,14 +440,19 @@ FString FEmitDefaultValueHelper::GenerateConstructor(UClass* InBPGC, const FGath
 		Context.AddLine(FString::Printf(TEXT("if(HasAnyFlags(RF_ClassDefaultObject) && (%s::StaticClass() == GetClass()))"), *CppClassName));
 		Context.AddLine(TEXT("{"));
 		Context.IncreaseIndent();
-		Context.AddLine(TEXT("ensure(0 == GetClass()->ConvertedSubobjectsFromBPGC.Num());"));
-		Context.bCreatingObjectsPerClass = true;
+		Context.AddLine(TEXT("ensure(0 == CastChecked<UDynamicClass>(GetClass())->MiscConvertedSubobjects.Num());"));
+		Context.AddLine(TEXT("ensure(0 == CastChecked<UDynamicClass>(GetClass())->ReferencedConvertedFields.Num());"));
+		Context.AddLine(TEXT("ensure(0 == CastChecked<UDynamicClass>(GetClass())->UsedAssets.Num());"));
+		Context.AddLine(TEXT("ensure(0 == CastChecked<UDynamicClass>(GetClass())->DynamicBindingObjects.Num());"));
+		Context.AddLine(TEXT("ensure(0 == CastChecked<UDynamicClass>(GetClass())->ComponentTemplates.Num());"));
+		Context.AddLine(TEXT("ensure(0 == CastChecked<UDynamicClass>(GetClass())->Timelines.Num());"));
+		Context.CurrentCodeType = FEmitterLocalContext::EGeneratedCodeType::SubobjectsOfClass;
 		Context.FlushLines();
 		for (auto ComponentTemplate : ActorComponentTempatesOwnedByClass)
 		{
 			if (ComponentTemplate)
 			{
-				HandleClassSubobject(Context, ComponentTemplate);
+				HandleClassSubobject(Context, ComponentTemplate, FEmitterLocalContext::EClassSubobjectList::ComponentTemplates);
 			}
 		}
 
@@ -455,7 +460,7 @@ FString FEmitDefaultValueHelper::GenerateConstructor(UClass* InBPGC, const FGath
 		{
 			if (TimelineTemplate)
 			{
-				HandleClassSubobject(Context, TimelineTemplate);
+				HandleClassSubobject(Context, TimelineTemplate, FEmitterLocalContext::EClassSubobjectList::Timelines);
 			}
 		}
 
@@ -463,176 +468,169 @@ FString FEmitDefaultValueHelper::GenerateConstructor(UClass* InBPGC, const FGath
 		{
 			if (DynamicBindingObject)
 			{
-				HandleClassSubobject(Context, DynamicBindingObject);
+				HandleClassSubobject(Context, DynamicBindingObject, FEmitterLocalContext::EClassSubobjectList::DynamicBindingObjects);
 			}
 		}
 
 		FBackendHelperUMG::CreateClassSubobjects(Context);
-
-		Context.bCreatingObjectsPerClass = false;
-
-		FDependenciesHelper::AddDependenciesInConstructor(Context, Dependencies);
+		FDependenciesHelper::AddDependenciesInConstructor(Context);
 
 		Context.DecreaseIndent();
 		Context.AddLine(TEXT("}"));
+	}
 
+	{
+		Context.CurrentCodeType = FEmitterLocalContext::EGeneratedCodeType::CommonConstructor;
 		// Let's have an easy access to generated class subobjects
 		Context.AddLine(TEXT("{")); // no shadow variables
 		Context.IncreaseIndent();
 		Context.FlushLines();
-		int32 SubobjectIndex = 0;
-		for (UObject* Obj : Context.GetObjectsCreatedPerClass_InConstructor())
-		{
-			FString NativeName;
-			const bool bFound = Context.FindLocalObject_InConstructor(Obj, NativeName);
-			ensure(bFound);
-			Context.AddHighPriorityLine(FString::Printf(TEXT("auto %s = CastChecked<%s>(%s::StaticClass()->ConvertedSubobjectsFromBPGC[%d]);")
-				, *NativeName
-				, *FEmitHelper::GetCppName(Obj->GetClass())
-				, *CppClassName
-				, SubobjectIndex));
-			SubobjectIndex++;
-		}
-	}
 
-	UObject* CDO = BPGC->GetDefaultObject(false);
-	Context.AddObjectFromLocalProperty_InConstructor(CDO, TEXT("this"));
+		UObject* CDO = BPGC->GetDefaultObject(false);
 
-	UObject* ParentCDO = BPGC->GetSuperClass()->GetDefaultObject(false);
-	check(CDO && ParentCDO);
-	Context.AddLine(TEXT(""));
+		UObject* ParentCDO = BPGC->GetSuperClass()->GetDefaultObject(false);
+		check(CDO && ParentCDO);
+		Context.AddLine(TEXT(""));
 
-	FString NativeRootComponentFallback;
-	TSet<const UProperty*> HandledProperties;
+		FString NativeRootComponentFallback;
+		TSet<const UProperty*> HandledProperties;
 
-	// Generate ctor init code for native class default subobjects that are always instanced (e.g. components).
-	// @TODO (pkavan) - We can probably make this faster by generating code to index through the DSO array instead (i.e. in place of HandleInstancedSubobject which will generate a lookup call per DSO).
-	TArray<UObject*> NativeDefaultObjectSubobjects;
-	BPGC->GetDefaultObjectSubobjects(NativeDefaultObjectSubobjects);
-	for (auto DSO : NativeDefaultObjectSubobjects)
-	{
-		if (DSO && DSO->GetClass()->HasAnyClassFlags(CLASS_DefaultToInstanced))
+		// Generate ctor init code for native class default subobjects that are always instanced (e.g. components).
+		// @TODO (pkavan) - We can probably make this faster by generating code to index through the DSO array instead (i.e. in place of HandleInstancedSubobject which will generate a lookup call per DSO).
+		TArray<UObject*> NativeDefaultObjectSubobjects;
+		BPGC->GetDefaultObjectSubobjects(NativeDefaultObjectSubobjects);
+		for (auto DSO : NativeDefaultObjectSubobjects)
 		{
-			// Determine if this is an editor-only subobject.
-			bool bIsEditorOnlySubobject = false;
-			if (const UActorComponent* ActorComponent = Cast<UActorComponent>(DSO))
-			{
-				bIsEditorOnlySubobject = ActorComponent->IsEditorOnly();
-			}
-
-			// Skip ctor code gen for editor-only subobjects, since they won't be used by the runtime. Any dependencies on editor-only subobjects will be handled later (see HandleInstancedSubobject).
-			if (!bIsEditorOnlySubobject)
+			if (DSO && DSO->GetClass()->HasAnyClassFlags(CLASS_DefaultToInstanced))
 			{
-				const FString VariableName = HandleInstancedSubobject(Context, DSO, false, true);
+				// Determine if this is an editor-only subobject.
+				bool bIsEditorOnlySubobject = false;
+				if (const UActorComponent* ActorComponent = Cast<UActorComponent>(DSO))
+				{
+					bIsEditorOnlySubobject = ActorComponent->IsEditorOnly();
+				}
 
-				// Keep track of which component can be used as a root, in case it's not explicitly set.
-				if (NativeRootComponentFallback.IsEmpty())
+				// Skip ctor code gen for editor-only subobjects, since they won't be used by the runtime. Any dependencies on editor-only subobjects will be handled later (see HandleInstancedSubobject).
+				if (!bIsEditorOnlySubobject)
 				{
-					USceneComponent* SceneComponent = Cast<USceneComponent>(DSO);
-					if (SceneComponent && !SceneComponent->AttachParent && SceneComponent->CreationMethod == EComponentCreationMethod::Native)
+					const FString VariableName = HandleInstancedSubobject(Context, DSO, false, true);
+
+					// Keep track of which component can be used as a root, in case it's not explicitly set.
+					if (NativeRootComponentFallback.IsEmpty())
 					{
-						NativeRootComponentFallback = VariableName;
+						USceneComponent* SceneComponent = Cast<USceneComponent>(DSO);
+						if (SceneComponent && !SceneComponent->AttachParent && SceneComponent->CreationMethod == EComponentCreationMethod::Native)
+						{
+							NativeRootComponentFallback = VariableName;
+						}
 					}
 				}
 			}
 		}
-	}
 
-	// Check for a valid RootComponent property value; mark it as handled if already set in the defaults.
-	bool bNeedsRootComponentAssignment = false;
-	static const FName RootComponentPropertyName(TEXT("RootComponent"));
-	const UObjectProperty* RootComponentProperty = FindField<UObjectProperty>(BPGC, RootComponentPropertyName);
-	if (RootComponentProperty)
-	{
-		if (RootComponentProperty->GetObjectPropertyValue_InContainer(CDO))
+		// Check for a valid RootComponent property value; mark it as handled if already set in the defaults.
+		bool bNeedsRootComponentAssignment = false;
+		static const FName RootComponentPropertyName(TEXT("RootComponent"));
+		const UObjectProperty* RootComponentProperty = FindField<UObjectProperty>(BPGC, RootComponentPropertyName);
+		if (RootComponentProperty)
 		{
-			HandledProperties.Add(RootComponentProperty);
-		}
-		else if (!NativeRootComponentFallback.IsEmpty())
-		{
-			Context.AddLine(FString::Printf(TEXT("RootComponent = %s;"), *NativeRootComponentFallback));
-			HandledProperties.Add(RootComponentProperty);
-		}
-		else
-		{
-			bNeedsRootComponentAssignment = true;
+			if (RootComponentProperty->GetObjectPropertyValue_InContainer(CDO))
+			{
+				HandledProperties.Add(RootComponentProperty);
+			}
+			else if (!NativeRootComponentFallback.IsEmpty())
+			{
+				Context.AddLine(FString::Printf(TEXT("RootComponent = %s;"), *NativeRootComponentFallback));
+				HandledProperties.Add(RootComponentProperty);
+			}
+			else
+			{
+				bNeedsRootComponentAssignment = true;
+			}
 		}
-	}
 
-	// Generate ctor init code for the SCS node hierarchy (i.e. non-native components). SCS nodes may have dependencies on native DSOs, but not vice-versa.
-	TArray<const UBlueprintGeneratedClass*> BPGCStack;
-	const bool bErrorFree = UBlueprintGeneratedClass::GetGeneratedClassesHierarchy(BPGC, BPGCStack);
-	if (bErrorFree)
-	{
-		// Start at the base of the hierarchy so that dependencies are handled first.
-		for (int32 i = BPGCStack.Num() - 1; i >= 0; --i)
+		// Generate ctor init code for the SCS node hierarchy (i.e. non-native components). SCS nodes may have dependencies on native DSOs, but not vice-versa.
+		TArray<const UBlueprintGeneratedClass*> BPGCStack;
+		const bool bErrorFree = UBlueprintGeneratedClass::GetGeneratedClassesHierarchy(BPGC, BPGCStack);
+		if (bErrorFree)
 		{
-			if (BPGCStack[i]->SimpleConstructionScript)
+			// Start at the base of the hierarchy so that dependencies are handled first.
+			for (int32 i = BPGCStack.Num() - 1; i >= 0; --i)
 			{
-				for (auto Node : BPGCStack[i]->SimpleConstructionScript->GetRootNodes())
+				if (BPGCStack[i]->SimpleConstructionScript)
 				{
-					if (Node)
+					for (auto Node : BPGCStack[i]->SimpleConstructionScript->GetRootNodes())
 					{
-						const FString VariableName = HandleNonNativeComponent(Context, Node, HandledProperties);
-
-						if (i == 0 && bNeedsRootComponentAssignment && Node->ComponentTemplate && Node->ComponentTemplate->IsA<USceneComponent>() && !VariableName.IsEmpty())
+						if (Node)
 						{
-							Context.AddLine(FString::Printf(TEXT("RootComponent = %s;"), *VariableName));
+							const FString NativeVariablePropertyName = HandleNonNativeComponent(Context, Node, HandledProperties);
+
+							if (i == 0 && bNeedsRootComponentAssignment && Node->ComponentTemplate && Node->ComponentTemplate->IsA<USceneComponent>() && !NativeVariablePropertyName.IsEmpty())
+							{
+								Context.AddLine(FString::Printf(TEXT("RootComponent = %s;"), *NativeVariablePropertyName));
 
-							bNeedsRootComponentAssignment = false;
-							HandledProperties.Add(RootComponentProperty);
+								bNeedsRootComponentAssignment = false;
+								HandledProperties.Add(RootComponentProperty);
+							}
 						}
 					}
 				}
 			}
 		}
-	}
 
-	// Generate ctor init code for generated Blueprint class property values that may differ from parent class defaults (or that otherwise belong to the generated Blueprint class).
-	for (auto Property : TFieldRange<const UProperty>(BPGC))
-	{
-		const bool bNewProperty = Property->GetOwnerStruct() == BPGC;
-		const bool bIsAccessible = bNewProperty || !Property->HasAnyPropertyFlags(CPF_NativeAccessSpecifierPrivate);
-		if (bIsAccessible && !HandledProperties.Contains(Property))
+		// Generate ctor init code for generated Blueprint class property values that may differ from parent class defaults (or that otherwise belong to the generated Blueprint class).
+		for (auto Property : TFieldRange<const UProperty>(BPGC))
 		{
-			OuterGenerate(Context, Property, TEXT(""), reinterpret_cast<const uint8*>(CDO), bNewProperty ? nullptr : reinterpret_cast<const uint8*>(ParentCDO), EPropertyAccessOperator::None, true);
+			const bool bNewProperty = Property->GetOwnerStruct() == BPGC;
+			const bool bIsAccessible = bNewProperty || !Property->HasAnyPropertyFlags(CPF_NativeAccessSpecifierPrivate);
+			if (bIsAccessible && !HandledProperties.Contains(Property))
+			{
+				OuterGenerate(Context, Property, TEXT(""), reinterpret_cast<const uint8*>(CDO), bNewProperty ? nullptr : reinterpret_cast<const uint8*>(ParentCDO), EPropertyAccessOperator::None, true);
+			}
 		}
-	}
 
+		Context.DecreaseIndent();
+		Context.AddLine(TEXT("}"));
+	}
 	Context.DecreaseIndent();
 	Context.AddLine(TEXT("}"));
-	Context.DecreaseIndent();
-	Context.AddLine(TEXT("}"));
-	Context.bInsideConstructor = false;
+	Context.CurrentCodeType = FEmitterLocalContext::EGeneratedCodeType::Regular;
 
-	FDependenciesHelper::AddStaticFunctionsForDependencies(Context, Dependencies);
+	FDependenciesHelper::AddStaticFunctionsForDependencies(Context);
 
 	FBackendHelperUMG::EmitWidgetInitializationFunctions(Context);
 
-	return Context.GetResult();
+	const FString Result = Context.GetResult();
+	Context.ClearResult();
+	return Result;
 }
 
-FString FEmitDefaultValueHelper::HandleClassSubobject(FEmitterLocalContext& Context, UObject* Object)
+FString FEmitDefaultValueHelper::HandleClassSubobject(FEmitterLocalContext& Context, UObject* Object, FEmitterLocalContext::EClassSubobjectList ListOfSubobjectsType)
 {
-	FString OuterStr;
-	if (!Context.FindLocalObject_InConstructor(Object->GetOuter(), OuterStr))
+	ensure(Context.CurrentCodeType == FEmitterLocalContext::EGeneratedCodeType::SubobjectsOfClass);
+	const FString OuterStr = Context.FindGloballyMappedObject(Object->GetOuter());
+	if (OuterStr.IsEmpty())
 	{
+		ensure(false);
 		return FString();
 	}
 
 	const bool AddAsSubobjectOfClass = Object->GetOuter() == Context.GetCurrentlyGeneratedClass();
-	const FString LocalNativeName = Context.AddNewObject_InConstructor(Object, AddAsSubobjectOfClass);
+	const FString LocalNativeName = Context.GenerateUniqueLocalName();
+	Context.AddClassSubObject_InConstructor(Object, LocalNativeName);
 	UClass* ObjectClass = Object->GetClass();
 	Context.AddHighPriorityLine(FString::Printf(
-		TEXT("auto %s = NewObject<%s%s>(%s, TEXT(\"%s\"));")
+		TEXT("auto %s = NewObject<%s>(%s, TEXT(\"%s\"));")
 		, *LocalNativeName
 		, *FEmitHelper::GetCppName(ObjectClass)
 		, *OuterStr
 		, *Object->GetName()));
 	if (AddAsSubobjectOfClass)
 	{
-		ensure(Context.bCreatingObjectsPerClass);
-		Context.AddHighPriorityLine(FString::Printf(TEXT("GetClass()->ConvertedSubobjectsFromBPGC.Add(%s);"), *LocalNativeName));
+		Context.RegisterClassSubobject(Object, ListOfSubobjectsType);
+		Context.AddHighPriorityLine(FString::Printf(TEXT("CastChecked<UDynamicClass>(GetClass())->%s.Add(%s);")
+			, Context.ClassSubobjectListName(ListOfSubobjectsType)
+			, *LocalNativeName));
 	}
 
 	for (auto Property : TFieldRange<const UProperty>(ObjectClass))
@@ -651,14 +649,23 @@ FString FEmitDefaultValueHelper::HandleInstancedSubobject(FEmitterLocalContext&
 	check(Object);
 
 	// Make sure we don't emit initialization code for the same object more than once.
-	FString LocalNativeName;
-	if (Context.FindLocalObject_InConstructor(Object, LocalNativeName))
+	FString LocalNativeName = Context.FindGloballyMappedObject(Object);
+	if (!LocalNativeName.IsEmpty())
 	{
 		return LocalNativeName;
 	}
 	else
 	{
-		LocalNativeName = Context.AddNewObject_InConstructor(Object, false);
+		LocalNativeName = Context.GenerateUniqueLocalName();
+	}
+
+	if (Context.CurrentCodeType == FEmitterLocalContext::EGeneratedCodeType::SubobjectsOfClass)
+	{
+		Context.AddClassSubObject_InConstructor(Object, LocalNativeName);
+	}
+	else if (Context.CurrentCodeType == FEmitterLocalContext::EGeneratedCodeType::CommonConstructor)
+	{
+		Context.AddCommonSubObject_InConstructor(Object, LocalNativeName);
 	}
 
 	UClass* ObjectClass = Object->GetClass();
@@ -689,8 +696,8 @@ FString FEmitDefaultValueHelper::HandleInstancedSubobject(FEmitterLocalContext&
 		}
 		else
 		{
-			Context.AddHighPriorityLine(FString::Printf(TEXT("auto %s = CastChecked<%s>(GetDefaultSubobjectByName(TEXT(\"%s\")));")
-				, *LocalNativeName, *FEmitHelper::GetCppName(ObjectClass), *Object->GetName()));
+		Context.AddHighPriorityLine(FString::Printf(TEXT("auto %s = CastChecked<%s>(GetDefaultSubobjectByName(TEXT(\"%s\")));")
+			, *LocalNativeName, *FEmitHelper::GetCppName(ObjectClass), *Object->GetName()));
 		}
 
 		const UObject* ObjectArchetype = Object->GetArchetype();
@@ -704,9 +711,10 @@ FString FEmitDefaultValueHelper::HandleInstancedSubobject(FEmitterLocalContext&
 	}
 	else
 	{
-		FString OuterStr;
-		if (!Context.FindLocalObject_InConstructor(Object->GetOuter(), OuterStr))
+		const FString OuterStr = Context.FindGloballyMappedObject(Object);
+		if (OuterStr.IsEmpty())
 		{
+			ensure(false);
 			return FString();
 		}
 		Context.AddHighPriorityLine(FString::Printf(TEXT("auto %s = NewObject<%s>(%s, TEXT(\"%s\"));")
diff --git a/Engine/Source/Runtime/CoreUObject/Private/UObject/Class.cpp b/Engine/Source/Runtime/CoreUObject/Private/UObject/Class.cpp
--- a/Engine/Source/Runtime/CoreUObject/Private/UObject/Class.cpp
+++ b/Engine/Source/Runtime/CoreUObject/Private/UObject/Class.cpp
@@ -2544,8 +2544,6 @@ void UClass::AddReferencedObjects(UObject* InThis, FReferenceCollector& Collecto
 		This->CallAddReferencedObjects(This->ClassDefaultObject, Collector);
 	}
 
-	Collector.AddReferencedObjects(This->ConvertedSubobjectsFromBPGC, This);
-
 	Super::AddReferencedObjects( This, Collector );
 }
 
@@ -4580,13 +4578,30 @@ UDynamicClass::UDynamicClass(
 
 void UDynamicClass::AddReferencedObjects(UObject* InThis, FReferenceCollector& Collector)
 {
-	UClass* This = CastChecked<UDynamicClass>(InThis);
+	UDynamicClass* This = CastChecked<UDynamicClass>(InThis);
 
-	//Collector.AddReferencedObjects(This->ConvertedSubobjectsFromBPGC, This);
+	Collector.AddReferencedObjects(This->MiscConvertedSubobjects, This);
+	Collector.AddReferencedObjects(This->ReferencedConvertedFields, This);
+	Collector.AddReferencedObjects(This->UsedAssets, This);
+	Collector.AddReferencedObjects(This->DynamicBindingObjects, This);
+	Collector.AddReferencedObjects(This->ComponentTemplates, This);
+	Collector.AddReferencedObjects(This->Timelines, This);
 
 	Super::AddReferencedObjects(This, Collector);
 }
 
+void UDynamicClass::PurgeClass(bool bRecompilingOnLoad)
+{
+	Super::PurgeClass(bRecompilingOnLoad);
+
+	MiscConvertedSubobjects.Empty();
+	ReferencedConvertedFields.Empty();
+	UsedAssets.Empty();
+
+	DynamicBindingObjects.Empty();
+	ComponentTemplates.Empty();
+	Timelines.Empty();
+}
 
 IMPLEMENT_CORE_INTRINSIC_CLASS(UDynamicClass, UClass,
 {
diff --git a/Engine/Source/Runtime/CoreUObject/Public/UObject/Class.h b/Engine/Source/Runtime/CoreUObject/Public/UObject/Class.h
--- a/Engine/Source/Runtime/CoreUObject/Public/UObject/Class.h
+++ b/Engine/Source/Runtime/CoreUObject/Public/UObject/Class.h
@@ -1947,9 +1947,6 @@ public:
 	/** This class's native functions. */
 	TArray<FNativeFunctionLookup> NativeFunctionLookupTable;
 
-	/** Misc objects owned by the class. Mostly templates from BP-generated native class */
-	TArray<UObject*> ConvertedSubobjectsFromBPGC;
-
 public:
 	// Constructors
 	UClass(const FObjectInitializer& ObjectInitializer = FObjectInitializer::Get());
@@ -2402,8 +2399,22 @@ public:
 	// UObject interface.
 	static void AddReferencedObjects(UObject* InThis, FReferenceCollector& Collector);
 
-	/** Misc objects owned by the class. Mostly templates from BP-generated native class */
-	//TArray<UObject*> ConvertedSubobjectsFromBPGC;
+	// UClass interface
+	virtual void PurgeClass(bool bRecompilingOnLoad) override;
+
+	/** Misc objects owned by the class. */
+	TArray<UObject*> MiscConvertedSubobjects;
+
+	/** Additional converted fields, that are used by the class. */
+	TArray<UField*> ReferencedConvertedFields;
+
+	/** Outer assets used by the class */
+	TArray<UObject*> UsedAssets;
+
+	// Specialized sub-object containers
+	TArray<UObject*> DynamicBindingObjects;
+	TArray<UObject*> ComponentTemplates;
+	TArray<UObject*> Timelines;
 };
 
 /**
diff --git a/Engine/Source/Runtime/Engine/Private/ActorConstruction.cpp b/Engine/Source/Runtime/Engine/Private/ActorConstruction.cpp
--- a/Engine/Source/Runtime/Engine/Private/ActorConstruction.cpp
+++ b/Engine/Source/Runtime/Engine/Private/ActorConstruction.cpp
@@ -619,9 +619,9 @@ UActorComponent* AActor::AddComponent(FName TemplateName, bool bManualAttachment
 		{
 			Template = BPGC->FindComponentTemplateByName(TemplateName);
 		}
-		else
+		else if (auto DynamicClass = Cast<UDynamicClass>(TemplateOwnerClass))
 		{
-			UObject** FoundTemplatePtr = TemplateOwnerClass->ConvertedSubobjectsFromBPGC.FindByPredicate([=](UObject* Obj) -> bool
+			UObject** FoundTemplatePtr = DynamicClass->ComponentTemplates.FindByPredicate([=](UObject* Obj) -> bool
 			{
 				return Obj && Obj->IsA<UActorComponent>() && (Obj->GetFName() == TemplateName);
 			});
diff --git a/Engine/Source/Runtime/Engine/Private/BlueprintGeneratedClass.cpp b/Engine/Source/Runtime/Engine/Private/BlueprintGeneratedClass.cpp
--- a/Engine/Source/Runtime/Engine/Private/BlueprintGeneratedClass.cpp
+++ b/Engine/Source/Runtime/Engine/Private/BlueprintGeneratedClass.cpp
@@ -343,9 +343,9 @@ UDynamicBlueprintBinding* UBlueprintGeneratedClass::GetDynamicBindingObject(cons
 			}
 		}
 	}
-	else
+	else if (auto DynamicClass = Cast<UDynamicClass>(ThisClass))
 	{
-		for (auto MiscObj : ThisClass->ConvertedSubobjectsFromBPGC)
+		for (auto MiscObj : DynamicClass->DynamicBindingObjects)
 		{
 			auto DynamicBindingObject = Cast<UDynamicBlueprintBinding>(MiscObj);
 			if (DynamicBindingObject && (DynamicBindingObject->GetClass() == BindingClass))
@@ -377,9 +377,9 @@ void UBlueprintGeneratedClass::BindDynamicDelegates(const UClass* ThisClass, UOb
 			}
 		}
 	}
-	else
+	else if (auto DynamicClass = Cast<UDynamicClass>(ThisClass))
 	{
-		for (auto MiscObj : ThisClass->ConvertedSubobjectsFromBPGC)
+		for (auto MiscObj : DynamicClass->DynamicBindingObjects)
 		{
 			auto DynamicBindingObject = Cast<UDynamicBlueprintBinding>(MiscObj);
 			if (DynamicBindingObject)
@@ -580,13 +580,16 @@ void UBlueprintGeneratedClass::CreateComponentsForActor(const UClass* ThisClass,
 			}
 		}
 
-		for (auto MiscObj : ThisClass->ConvertedSubobjectsFromBPGC)
+		if (auto DynamicClass = Cast<UDynamicClass>(ThisClass))
 		{
-			auto TimelineTemplate = Cast<const UTimelineTemplate>(MiscObj);
-			// Not fatal if NULL, but shouldn't happen and ignored if not wired up in graph
-			if (TimelineTemplate && TimelineTemplate->bValidatedAsWired)
+			for (auto MiscObj : DynamicClass->Timelines)
 			{
-				CreateTimelineComponent(Actor, TimelineTemplate);
+				auto TimelineTemplate = Cast<const UTimelineTemplate>(MiscObj);
+				// Not fatal if NULL, but shouldn't happen and ignored if not wired up in graph
+				if (TimelineTemplate && TimelineTemplate->bValidatedAsWired)
+				{
+					CreateTimelineComponent(Actor, TimelineTemplate);
+				}
 			}
 		}
 	}
diff --git a/Engine/Source/Runtime/Engine/Private/InputDelegateBinding.cpp b/Engine/Source/Runtime/Engine/Private/InputDelegateBinding.cpp
--- a/Engine/Source/Runtime/Engine/Private/InputDelegateBinding.cpp
+++ b/Engine/Source/Runtime/Engine/Private/InputDelegateBinding.cpp
@@ -17,8 +17,9 @@ UInputDelegateBinding::UInputDelegateBinding(const FObjectInitializer& ObjectIni
 
 bool UInputDelegateBinding::SupportsInputDelegate(const UClass* InClass)
 {
+	auto DynamicClass = Cast<UDynamicClass>(InClass);
 	return InClass && (
-		InClass->ConvertedSubobjectsFromBPGC.Num()
+		(DynamicClass && DynamicClass->DynamicBindingObjects.Num())
 		|| Cast<UBlueprintGeneratedClass>(InClass)
 		|| SupportsInputDelegate(InClass->GetSuperClass()));
 }
