UEBP-40 C++ code generation (WIP):
- Test Pipeline for native classes converted from BP (FReplaceConvertedAssetManager)
- DynamicClass support

diff --git a/Engine/Config/BaseEngine.ini b/Engine/Config/BaseEngine.ini
--- a/Engine/Config/BaseEngine.ini
+++ b/Engine/Config/BaseEngine.ini
@@ -1700,6 +1700,7 @@ bForceDisableCookedDependencyDeferring=false
 bExecutionAfterReturn=false
 bCanSuppressAccessViolation=false
 bUseSelectRef=true
+bLoadNativeConvertedBPClassInEditor=false
 
 [/Script/Engine.Blueprint]
 bRecompileOnLoad=true
diff --git a/Engine/Source/Developer/BlueprintCompilerCppBackend/Private/BlueprintCompilerCppBackendValueHelper.cpp b/Engine/Source/Developer/BlueprintCompilerCppBackend/Private/BlueprintCompilerCppBackendValueHelper.cpp
--- a/Engine/Source/Developer/BlueprintCompilerCppBackend/Private/BlueprintCompilerCppBackendValueHelper.cpp
+++ b/Engine/Source/Developer/BlueprintCompilerCppBackend/Private/BlueprintCompilerCppBackendValueHelper.cpp
@@ -387,14 +387,44 @@ struct FDependenciesHelper
 		Context.AddLine(TEXT("{"));
 		Context.IncreaseIndent();
 
+		{
+			// __StaticDependenciesAssets
+			Context.DecreaseIndent();
+			Context.AddLine(TEXT("#if WITH_EDITOR"));
+			Context.IncreaseIndent();
+			Context.AddLine(TEXT("static UField* __StaticCreateClass()"));
+			Context.AddLine(TEXT("{"));
+			Context.IncreaseIndent();
+			Context.AddLine(FString::Printf(TEXT("return CastChecked<UDynamicClass>(%s::StaticClass());"), *CppClassName));
+			Context.DecreaseIndent();
+			Context.AddLine(TEXT("}"));
+			Context.DecreaseIndent();
+			Context.AddLine(TEXT("#endif //WITH_EDITOR"));
+			Context.IncreaseIndent();
+		}
+
 		Context.AddLine(FString::Printf(TEXT("%s()"), *RegisterHelperName));
 		Context.AddLine(TEXT("{"));
 		Context.IncreaseIndent();
+
 		Context.AddLine(FString::Printf(
 			TEXT("FConvertedBlueprintsDependencies::Get().RegisterClass(TEXT(\"%s\"), &%s::__StaticDependenciesAssets);")
 			, *OriginalClass->GetName()
 			, *CppClassName
 			, *CppClassName));
+
+		{
+			Context.DecreaseIndent();
+			Context.AddLine(TEXT("#if WITH_EDITOR"));
+			Context.IncreaseIndent();
+
+			Context.AddLine(TEXT("FConvertedBlueprintsDependencies::Get().RegisterClass(&__StaticCreateClass);"));
+
+			Context.DecreaseIndent();
+			Context.AddLine(TEXT("#endif //WITH_EDITOR"));
+			Context.IncreaseIndent();
+		}
+
 		Context.DecreaseIndent();
 		Context.AddLine(TEXT("}"));
 
diff --git a/Engine/Source/Runtime/CoreUObject/Private/Blueprint/BlueprintSupport.cpp b/Engine/Source/Runtime/CoreUObject/Private/Blueprint/BlueprintSupport.cpp
--- a/Engine/Source/Runtime/CoreUObject/Private/Blueprint/BlueprintSupport.cpp
+++ b/Engine/Source/Runtime/CoreUObject/Private/Blueprint/BlueprintSupport.cpp
@@ -1663,7 +1663,8 @@ FReplaceConvertedAssetManager::FReplaceConvertedAssetManager()
 	FModuleManager::Get().OnModulesChanged().AddStatic(&FReplaceConvertedAssetManager::OnModulesChanged);
 
 	// FOR DEVELOPMENT/TEST ONLY:
-	// SetEnabled(true);
+	const FBoolConfigValueHelper LoadNativeConvertedBPClass(TEXT("Kismet"), TEXT("bLoadNativeConvertedBPClassInEditor"), GEngineIni);
+	SetEnabled(LoadNativeConvertedBPClass);
 }
 
 void FReplaceConvertedAssetManager::OnModulesChanged(FName ModuleThatChanged, EModuleChangeReason ReasonForChange)
@@ -1696,6 +1697,11 @@ UPackage* FReplaceConvertedAssetManager::FindPackageReplacement(const FString& O
 	return nullptr;
 }
 
+void FReplaceConvertedAssetManager::AddReferencedObjects(FReferenceCollector& Collector)
+{
+	Collector.AddReferencedObjects(ReplaceMap);
+}
+
 void FReplaceConvertedAssetManager::GatherOriginalPathsOfConvertedAssets()
 {
 	if (!IsEnabled())
@@ -1720,9 +1726,10 @@ void FReplaceConvertedAssetManager::GatherOriginalPathsOfConvertedAssets()
 		}
 	};
 
-	for (UClass* LocalClass : TObjectRange<UClass>())
+	for (auto FuncIter : FConvertedBlueprintsDependencies::Get().CreateClassFunctions)
 	{
-		FillMap(LocalClass);
+		auto DynamicClass = (*FuncIter)();
+		FillMap(DynamicClass);
 	}
 
 	for (UScriptStruct* LocalScriptStruct : TObjectRange<UScriptStruct>())
diff --git a/Engine/Source/Runtime/CoreUObject/Private/UObject/Obj.cpp b/Engine/Source/Runtime/CoreUObject/Private/UObject/Obj.cpp
--- a/Engine/Source/Runtime/CoreUObject/Private/UObject/Obj.cpp
+++ b/Engine/Source/Runtime/CoreUObject/Private/UObject/Obj.cpp
@@ -867,6 +867,11 @@ void UObject::Serialize( FArchive& Ar )
 	UObject* LoadOuter = GetOuter();
 	FName LoadName = GetFName();
 
+	if (Ar.IsLoading() && Cast<UDynamicClass>(ObjClass) && ObjClass->GetName().Contains(TEXT("RRRActor_C")))
+	{
+		int k = 0; k++;
+	}
+
 	// Make sure this object's class's data is loaded.
 	if(ObjClass->HasAnyFlags(RF_NeedLoad) )
 	{
diff --git a/Engine/Source/Runtime/CoreUObject/Public/Blueprint/BlueprintSupport.h b/Engine/Source/Runtime/CoreUObject/Public/Blueprint/BlueprintSupport.h
--- a/Engine/Source/Runtime/CoreUObject/Public/Blueprint/BlueprintSupport.h
+++ b/Engine/Source/Runtime/CoreUObject/Public/Blueprint/BlueprintSupport.h
@@ -69,7 +69,7 @@ private:
 * Native entities (generated from BP items) have "ReplaceConverted" with original object path.
 * This path is used to fix linker.
 */
-struct COREUOBJECT_API FReplaceConvertedAssetManager
+struct COREUOBJECT_API FReplaceConvertedAssetManager : public FGCObject
 {
 private:
 	TMap<FString, UObject*> ReplaceMap;
@@ -97,6 +97,8 @@ public:
 		return bIsEnabled;
 	}
 
+	virtual void AddReferencedObjects(FReferenceCollector& Collector) override;
+
 	UPackage* FindPackageReplacement(const FString& OriginalPathName) const;
 
 	UObject* FindReplacement(const FString& OriginalPathName) const;
@@ -169,6 +171,19 @@ private:
 
 public:
 
+#if WITH_EDITOR
+
+	typedef UField*(*StaticCreateClassFunc)();
+
+	TArray<StaticCreateClassFunc> CreateClassFunctions;
+
+	void RegisterClass(StaticCreateClassFunc CreateClass)
+	{
+		CreateClassFunctions.Add(CreateClass);
+	}
+
+#endif //WITH_EDITOR
+
 	static FConvertedBlueprintsDependencies& Get();
 
 	bool AnyClassRegistered() const
