UEBP-40 C++ code generation (WIP):
- components created in constructor have CreationMethod = EComponentCreationMethod::Native

diff --git a/Engine/Source/Developer/BlueprintCompilerCppBackend/Private/BlueprintCompilerCppBackendBase.cpp b/Engine/Source/Developer/BlueprintCompilerCppBackend/Private/BlueprintCompilerCppBackendBase.cpp
--- a/Engine/Source/Developer/BlueprintCompilerCppBackend/Private/BlueprintCompilerCppBackendBase.cpp
+++ b/Engine/Source/Developer/BlueprintCompilerCppBackend/Private/BlueprintCompilerCppBackendBase.cpp
@@ -138,6 +138,8 @@ void FBlueprintCompilerCppBackendBase::GenerateCodeFromClass(UClass* SourceClass
 		Emit(Body, *FEmitDefaultValueHelper::GenerateConstructor(EmitterContext));
 	}
 
+	Emit(Header, TEXT("\n\tvirtual void PostLoadSubobjects(FObjectInstancingGraph* OuterInstanceGraph) override;\n"));
+
 	for (int32 i = 0; i < Functions.Num(); ++i)
 	{
 		if (Functions[i].IsValid())
diff --git a/Engine/Source/Developer/BlueprintCompilerCppBackend/Private/BlueprintCompilerCppBackendUtils.h b/Engine/Source/Developer/BlueprintCompilerCppBackend/Private/BlueprintCompilerCppBackendUtils.h
--- a/Engine/Source/Developer/BlueprintCompilerCppBackend/Private/BlueprintCompilerCppBackendUtils.h
+++ b/Engine/Source/Developer/BlueprintCompilerCppBackend/Private/BlueprintCompilerCppBackendUtils.h
@@ -249,7 +249,7 @@ private:
 	// returns empty string if cannot handle
 	static FString HandleSpecialTypes(FEmitterLocalContext& Context, const UProperty* Property, const uint8* ValuePtr);
 
-	static FString HandleNonNativeComponent(FEmitterLocalContext& Context, const USCS_Node* Node, TSet<const UProperty*>& OutHandledProperties, const USCS_Node* ParentNode = nullptr);
+	static FString HandleNonNativeComponent(FEmitterLocalContext& Context, const USCS_Node* Node, TSet<const UProperty*>& OutHandledProperties, TArray<FString>& NativeCreatedComponentProperties, const USCS_Node* ParentNode = nullptr);
 
 	static FString HandleInstancedSubobject(FEmitterLocalContext& Context, UObject* Object, bool bCreateInstance = true, bool bSkipEditorOnlyCheck = false);
 };
diff --git a/Engine/Source/Developer/BlueprintCompilerCppBackend/Private/BlueprintCompilerCppBackendValueHelper.cpp b/Engine/Source/Developer/BlueprintCompilerCppBackend/Private/BlueprintCompilerCppBackendValueHelper.cpp
--- a/Engine/Source/Developer/BlueprintCompilerCppBackend/Private/BlueprintCompilerCppBackendValueHelper.cpp
+++ b/Engine/Source/Developer/BlueprintCompilerCppBackend/Private/BlueprintCompilerCppBackendValueHelper.cpp
@@ -245,7 +245,7 @@ FString FEmitDefaultValueHelper::HandleSpecialTypes(FEmitterLocalContext& Contex
 	return FString();
 }
 
-FString FEmitDefaultValueHelper::HandleNonNativeComponent(FEmitterLocalContext& Context, const USCS_Node* Node, TSet<const UProperty*>& OutHandledProperties, const USCS_Node* ParentNode)
+FString FEmitDefaultValueHelper::HandleNonNativeComponent(FEmitterLocalContext& Context, const USCS_Node* Node, TSet<const UProperty*>& OutHandledProperties, TArray<FString>& NativeCreatedComponentProperties, const USCS_Node* ParentNode)
 {
 	check(Node);
 	check(Context.CurrentCodeType == FEmitterLocalContext::EGeneratedCodeType::CommonConstructor);
@@ -286,9 +286,10 @@ FString FEmitDefaultValueHelper::HandleNonNativeComponent(FEmitterLocalContext&
 					, (VariableProperty == nullptr) ? TEXT("auto ") : TEXT("")
 					, *NativeVariablePropertyName
 					, *FEmitHelper::GetCppName(ComponentClass)
-					, *ComponentTemplate->GetName()));
+					, *VariableCleanName));
 
-				Context.AddLine(FString::Printf(TEXT("%s->CreationMethod = EComponentCreationMethod::SimpleConstructionScript;"), *NativeVariablePropertyName));
+				Context.AddLine(FString::Printf(TEXT("%s->CreationMethod = EComponentCreationMethod::Native;"), *NativeVariablePropertyName));
+				NativeCreatedComponentProperties.Add(NativeVariablePropertyName);
 
 				FString ParentVariableName;
 				if (ParentNode)
@@ -321,7 +322,7 @@ FString FEmitDefaultValueHelper::HandleNonNativeComponent(FEmitterLocalContext&
 	// Recursively handle child nodes.
 	for (auto ChildNode : Node->ChildNodes)
 	{
-		HandleNonNativeComponent(Context, ChildNode, OutHandledProperties, Node);
+		HandleNonNativeComponent(Context, ChildNode, OutHandledProperties, NativeCreatedComponentProperties, Node);
 	}
 
 	return NativeVariablePropertyName;
@@ -509,6 +510,9 @@ FString FEmitDefaultValueHelper::GenerateConstructor(FEmitterLocalContext& Conte
 		Context.AddLine(TEXT("}"));
 	}
 
+	// Components that must be fixed after serialization
+	TArray<FString> NativeCreatedComponentProperties;
+
 	{
 		Context.CurrentCodeType = FEmitterLocalContext::EGeneratedCodeType::CommonConstructor;
 		// Let's have an easy access to generated class subobjects
@@ -593,7 +597,7 @@ FString FEmitDefaultValueHelper::GenerateConstructor(FEmitterLocalContext& Conte
 					{
 						if (Node)
 						{
-							const FString NativeVariablePropertyName = HandleNonNativeComponent(Context, Node, HandledProperties);
+							const FString NativeVariablePropertyName = HandleNonNativeComponent(Context, Node, HandledProperties, NativeCreatedComponentProperties);
 
 							if (i == 0 && bNeedsRootComponentAssignment && Node->ComponentTemplate && Node->ComponentTemplate->IsA<USceneComponent>() && !NativeVariablePropertyName.IsEmpty())
 							{
@@ -626,6 +630,25 @@ FString FEmitDefaultValueHelper::GenerateConstructor(FEmitterLocalContext& Conte
 	Context.AddLine(TEXT("}"));
 	Context.CurrentCodeType = FEmitterLocalContext::EGeneratedCodeType::Regular;
 
+	{
+		Context.AddLine(FString::Printf(TEXT("void %s::PostLoadSubobjects(FObjectInstancingGraph* OuterInstanceGraph)"), *CppClassName));
+		Context.AddLine(TEXT("{"));
+		Context.IncreaseIndent();
+		Context.AddLine(TEXT("Super::PostLoadSubobjects(OuterInstanceGraph);"));
+		
+		for (auto& ComponentToFix : NativeCreatedComponentProperties)
+		{
+			Context.AddLine(FString::Printf(TEXT("if(ensure(%s))"), *ComponentToFix));
+			Context.AddLine(TEXT("{"));
+			Context.IncreaseIndent();
+			Context.AddLine(FString::Printf(TEXT("%s->CreationMethod = EComponentCreationMethod::Native;"), *ComponentToFix));
+			Context.DecreaseIndent();
+			Context.AddLine(TEXT("}"));
+		}
+		Context.DecreaseIndent();
+		Context.AddLine(TEXT("}"));
+	}
+
 	FDependenciesHelper::AddStaticFunctionsForDependencies(Context);
 
 	FBackendHelperUMG::EmitWidgetInitializationFunctions(Context);
