BlueprintCompilerCppBackend: improved interface generation.

diff --git a/Engine/Source/Developer/BlueprintCompilerCppBackend/Private/BlueprintCompilerCppBackendBase.cpp b/Engine/Source/Developer/BlueprintCompilerCppBackend/Private/BlueprintCompilerCppBackendBase.cpp
--- a/Engine/Source/Developer/BlueprintCompilerCppBackend/Private/BlueprintCompilerCppBackendBase.cpp
+++ b/Engine/Source/Developer/BlueprintCompilerCppBackend/Private/BlueprintCompilerCppBackendBase.cpp
@@ -136,9 +136,10 @@ void FBlueprintCompilerCppBackendBase::GenerateCodeFromClass(UClass* SourceClass
 	{
 		Emit(Header, *FString::Printf(TEXT("\t%s(const FObjectInitializer& ObjectInitializer = FObjectInitializer::Get());\n\n"), *CppClassName));
 		Emit(Body, *FEmitDefaultValueHelper::GenerateConstructor(EmitterContext));
-	}
 
-	Emit(Header, TEXT("\n\tvirtual void PostLoadSubobjects(FObjectInstancingGraph* OuterInstanceGraph) override;\n"));
+		Emit(Header, TEXT("\n\tvirtual void PostLoadSubobjects(FObjectInstancingGraph* OuterInstanceGraph) override;\n"));
+		Emit(Header, TEXT("\n\tstatic void __StaticDependenciesAssets(TArray<FName>& OutPackagePaths);\n"));
+	}
 
 	for (int32 i = 0; i < Functions.Num(); ++i)
 	{
@@ -148,8 +149,6 @@ void FBlueprintCompilerCppBackendBase::GenerateCodeFromClass(UClass* SourceClass
 		}
 	}
 
-	Emit(Header, TEXT("\n\tstatic void __StaticDependenciesAssets(TArray<FName>& OutPackagePaths);\n"));
-
 	Emit(Header, *FBackendHelperUMG::WidgetFunctionsInHeader(SourceClass));
 
 	Emit(Header, TEXT("};\n\n"));
diff --git a/Engine/Source/Developer/BlueprintCompilerCppBackend/Private/BlueprintCompilerCppBackendUtils.cpp b/Engine/Source/Developer/BlueprintCompilerCppBackend/Private/BlueprintCompilerCppBackendUtils.cpp
--- a/Engine/Source/Developer/BlueprintCompilerCppBackend/Private/BlueprintCompilerCppBackendUtils.cpp
+++ b/Engine/Source/Developer/BlueprintCompilerCppBackend/Private/BlueprintCompilerCppBackendUtils.cpp
@@ -86,7 +86,7 @@ FString FEmitterLocalContext::FindGloballyMappedObject(UObject* Object, bool bLo
 		auto BPGC = Cast<UBlueprintGeneratedClass>(ObjClass);
 		if (ObjClass->HasAnyClassFlags(CLASS_Native) || (BPGC && Dependencies.WillClassBeConverted(BPGC)))
 		{
-			return FString::Printf(TEXT("%s::StaticClass()"), *FEmitHelper::GetCppName(ObjClass));
+			return FString::Printf(TEXT("%s::StaticClass()"), *FEmitHelper::GetCppName(ObjClass, true));
 		}
 	}
 
@@ -119,7 +119,7 @@ FString FEmitterLocalContext::FindGloballyMappedObject(UObject* Object, bool bLo
 	return FString{};
 }
 
-FString FEmitHelper::GetCppName(const UField* Field)
+FString FEmitHelper::GetCppName(const UField* Field, bool bUInterface)
 {
 	check(Field);
 	auto AsClass = Cast<UClass>(Field);
@@ -129,7 +129,9 @@ FString FEmitHelper::GetCppName(const UField* Field)
 		if (AsClass && AsClass->HasAnyClassFlags(CLASS_Interface))
 		{
 			ensure(AsClass->IsChildOf<UInterface>());
-			return FString::Printf(TEXT("I%s"), *AsClass->GetName());
+			return FString::Printf(TEXT("%s%s")
+				, bUInterface ? TEXT("U") : TEXT("I")
+				, *AsClass->GetName());
 		}
 		auto AsStruct = CastChecked<UStruct>(Field);
 		return FString::Printf(TEXT("%s%s"), AsStruct->GetPrefixCPP(), *AsStruct->GetName());
diff --git a/Engine/Source/Developer/BlueprintCompilerCppBackend/Private/BlueprintCompilerCppBackendUtils.h b/Engine/Source/Developer/BlueprintCompilerCppBackend/Private/BlueprintCompilerCppBackendUtils.h
--- a/Engine/Source/Developer/BlueprintCompilerCppBackend/Private/BlueprintCompilerCppBackendUtils.h
+++ b/Engine/Source/Developer/BlueprintCompilerCppBackend/Private/BlueprintCompilerCppBackendUtils.h
@@ -175,7 +175,8 @@ public:
 
 struct FEmitHelper
 {
-	static FString GetCppName(const UField* Field);
+	// bUInterface - use interface with "U" prefix, by default there is "I" prefix
+	static FString GetCppName(const UField* Field, bool bUInterface = false);
 
 	static void ArrayToString(const TArray<FString>& Array, FString& OutString, const TCHAR* Separator);
 
diff --git a/Engine/Source/Developer/BlueprintCompilerCppBackend/Private/BlueprintCompilerCppBackendValueHelper.cpp b/Engine/Source/Developer/BlueprintCompilerCppBackend/Private/BlueprintCompilerCppBackendValueHelper.cpp
--- a/Engine/Source/Developer/BlueprintCompilerCppBackend/Private/BlueprintCompilerCppBackendValueHelper.cpp
+++ b/Engine/Source/Developer/BlueprintCompilerCppBackend/Private/BlueprintCompilerCppBackendValueHelper.cpp
@@ -339,7 +339,7 @@ struct FDependenciesHelper
 		for (auto LocStruct : Context.Dependencies.ConvertedClasses)
 		{
 			Context.AddLine(FString::Printf(TEXT("CastChecked<UDynamicClass>(GetClass())->ReferencedConvertedFields.Add(%s::StaticClass());")
-				, *FEmitHelper::GetCppName(LocStruct)));
+				, *FEmitHelper::GetCppName(LocStruct, true)));
 		}
 
 		if (Context.Dependencies.ConvertedStructs.Num())
@@ -396,7 +396,7 @@ struct FDependenciesHelper
 			Context.AddLine(TEXT("static UField* __StaticCreateClass()"));
 			Context.AddLine(TEXT("{"));
 			Context.IncreaseIndent();
-			Context.AddLine(FString::Printf(TEXT("return CastChecked<UDynamicClass>(%s::StaticClass());"), *CppClassName));
+			Context.AddLine(FString::Printf(TEXT("return CastChecked<UDynamicClass>(%s::StaticClass());"), *FEmitHelper::GetCppName(OriginalClass, true)));
 			Context.DecreaseIndent();
 			Context.AddLine(TEXT("}"));
 			Context.DecreaseIndent();
