FGatherConvertedClassDependencies collectsheaders for used interfaces.

diff --git a/Engine/Source/Developer/BlueprintCompilerCppBackend/Private/BlueprintCompilerCppBackendGatherDependencies.cpp b/Engine/Source/Developer/BlueprintCompilerCppBackend/Private/BlueprintCompilerCppBackendGatherDependencies.cpp
--- a/Engine/Source/Developer/BlueprintCompilerCppBackend/Private/BlueprintCompilerCppBackendGatherDependencies.cpp
+++ b/Engine/Source/Developer/BlueprintCompilerCppBackend/Private/BlueprintCompilerCppBackendGatherDependencies.cpp
@@ -27,17 +27,20 @@ void FGatherConvertedClassDependencies::DependenciesForHeader()
 
 	for (auto Obj : ObjectsToCheck)
 	{
-		auto Property = Cast<UProperty>(Obj);
+		auto Property = Cast<const UProperty>(Obj);
 		auto OwnerProperty = IsValid(Property) ? Property->GetOwnerProperty() : nullptr;
 		const bool bIsParam = OwnerProperty && (0 != (OwnerProperty->PropertyFlags & CPF_Parm)) && OwnerProperty->IsIn(OriginalStruct);
 		const bool bIsMemberVariable = OwnerProperty && (OwnerProperty->GetOuter() == OriginalStruct);
 		if (bIsParam || bIsMemberVariable)
 		{
-			const UObjectProperty* ObjectProperty = Cast<UObjectProperty>(OwnerProperty);
-			if (ObjectProperty)
+			if (auto ObjectProperty = Cast<const UObjectPropertyBase>(Property))
 			{
 				DeclareInHeader.Add(ObjectProperty->PropertyClass);
 			}
+			else if (auto InterfaceProperty = Cast<const UInterfaceProperty>(Property))
+			{
+				IncludeInHeader.Add(InterfaceProperty->InterfaceClass);
+			}
 			else
 			{
 				HeaderReferenceFinder.FindReferences(Obj);
