Coordinating the Blueprint conversion tool with the backend (ensuring filenames are in sync, and that the backend knows what files will and won't be converted).

diff --git a/Engine/Source/Developer/BlueprintCompilerCppBackend/Private/BlueprintCompilerCppBackendBase.cpp b/Engine/Source/Developer/BlueprintCompilerCppBackend/Private/BlueprintCompilerCppBackendBase.cpp
--- a/Engine/Source/Developer/BlueprintCompilerCppBackend/Private/BlueprintCompilerCppBackendBase.cpp
+++ b/Engine/Source/Developer/BlueprintCompilerCppBackend/Private/BlueprintCompilerCppBackendBase.cpp
@@ -73,7 +73,8 @@ void FBlueprintCompilerCppBackendBase::DeclareDelegates(UClass* SourceClass, TIn
 
 void FBlueprintCompilerCppBackendBase::GenerateCodeFromClass(UClass* SourceClass, TIndirectArray<FKismetFunctionContext>& Functions, bool bGenerateStubsOnly)
 {
-	auto CleanCppClassName = SourceClass->GetName();
+	// use GetBaseFilename() so that we can coordinate #includes and filenames
+	auto CleanCppClassName = FEmitHelper::GetBaseFilename(SourceClass);
 	CppClassName = FEmitHelper::GetCppName(SourceClass);
 	
 	FGatherConvertedClassDependencies Dependencies(SourceClass);
@@ -339,7 +340,8 @@ void FBlueprintCompilerCppBackendBase::ConstructFunction(FKismetFunctionContext&
 void FBlueprintCompilerCppBackendBase::GenerateCodeFromEnum(UUserDefinedEnum* SourceEnum)
 {
 	check(SourceEnum);
-	EmitFileBeginning(SourceEnum->GetName(), nullptr);
+	// use GetBaseFilename() so that we can coordinate #includes and filenames
+	EmitFileBeginning(FEmitHelper::GetBaseFilename(SourceEnum), nullptr);
 
 	Emit(Header, TEXT("UENUM(BlueprintType"));
 	EmitReplaceConvertedMetaData(SourceEnum);
@@ -376,7 +378,8 @@ void FBlueprintCompilerCppBackendBase::GenerateCodeFromStruct(UUserDefinedStruct
 	check(SourceStruct);
 
 	FGatherConvertedClassDependencies Dependencies(SourceStruct);
-	EmitFileBeginning(SourceStruct->GetName(), &Dependencies);
+	// use GetBaseFilename() so that we can coordinate #includes and filenames
+	EmitFileBeginning(FEmitHelper::GetBaseFilename(SourceStruct), &Dependencies);
 
 	const FString NewName = FEmitHelper::GetCppName(SourceStruct);
 	Emit(Header, TEXT("USTRUCT(BlueprintType"));
diff --git a/Engine/Source/Developer/BlueprintCompilerCppBackend/Private/BlueprintCompilerCppBackendModule.cpp b/Engine/Source/Developer/BlueprintCompilerCppBackend/Private/BlueprintCompilerCppBackendModule.cpp
--- a/Engine/Source/Developer/BlueprintCompilerCppBackend/Private/BlueprintCompilerCppBackendModule.cpp
+++ b/Engine/Source/Developer/BlueprintCompilerCppBackend/Private/BlueprintCompilerCppBackendModule.cpp
@@ -4,6 +4,7 @@
 #include "IBlueprintCompilerCppBackendModule.h"
 #include "ModuleManager.h"
 #include "BlueprintCompilerCppBackend.h"
+#include "BlueprintCompilerCppBackendUtils.h" // for FEmitHelper::GetBaseFilename()
 
 class FBlueprintCompilerCppBackendModule : public IBlueprintCompilerCppBackendModule
 {
@@ -33,3 +34,9 @@ IBlueprintCompilerCppBackend* FBlueprintCompilerCppBackendModule::Create()
 {
 	return new FBlueprintCompilerCppBackend();
 }
+
+FString IBlueprintCompilerCppBackendModule::GetBaseFilename(const UObject* AssetObj)
+{
+	// use the same function that the backend uses for #includes
+	return FEmitHelper::GetBaseFilename(AssetObj);
+}
diff --git a/Engine/Source/Developer/BlueprintCompilerCppBackend/Private/BlueprintCompilerCppBackendUtils.cpp b/Engine/Source/Developer/BlueprintCompilerCppBackend/Private/BlueprintCompilerCppBackendUtils.cpp
--- a/Engine/Source/Developer/BlueprintCompilerCppBackend/Private/BlueprintCompilerCppBackendUtils.cpp
+++ b/Engine/Source/Developer/BlueprintCompilerCppBackend/Private/BlueprintCompilerCppBackendUtils.cpp
@@ -411,6 +411,13 @@ FString FEmitHelper::GenerateReplaceConvertedMD(UObject* Obj)
 	return Result;
 }
 
+FString FEmitHelper::GetBaseFilename(const UObject* AssetObj)
+{
+	// @TODO: See DanO about properly handling special asset names (w/ unicode, conflicts, etc.)
+	FString PackagePath = AssetObj->GetOutermost()->GetPathName();
+	return FPackageName::GetLongPackageAssetName(PackagePath);
+}
+
 FString FEmitHelper::EmitUFuntion(UFunction* Function, TArray<FString>* AdditinalMetaData)
 {
 	TArray<FString> Tags = FEmitHelper::FunctionFlagsToTags(Function->FunctionFlags);
diff --git a/Engine/Source/Developer/BlueprintCompilerCppBackend/Private/BlueprintCompilerCppBackendUtils.h b/Engine/Source/Developer/BlueprintCompilerCppBackend/Private/BlueprintCompilerCppBackendUtils.h
--- a/Engine/Source/Developer/BlueprintCompilerCppBackend/Private/BlueprintCompilerCppBackendUtils.h
+++ b/Engine/Source/Developer/BlueprintCompilerCppBackend/Private/BlueprintCompilerCppBackendUtils.h
@@ -219,6 +219,8 @@ struct FEmitHelper
 	static bool GenerateAutomaticCast(const FEdGraphPinType& LType, const FEdGraphPinType& RType, FString& OutCastBegin, FString& OutCastEnd);
 
 	static FString GenerateReplaceConvertedMD(UObject* Obj);
+
+	static FString GetBaseFilename(const UObject* AssetObj);
 };
 
 struct FEmitDefaultValueHelper
@@ -265,4 +267,4 @@ struct FBackendHelperUMG
 	static void CreateClassSubobjects(FEmitterLocalContext& Context);
 	static void EmitWidgetInitializationFunctions(FEmitterLocalContext& Context);
 
-};
\ No newline at end of file
+};
diff --git a/Engine/Source/Developer/BlueprintCompilerCppBackend/Public/IBlueprintCompilerCppBackendModule.h b/Engine/Source/Developer/BlueprintCompilerCppBackend/Public/IBlueprintCompilerCppBackendModule.h
--- a/Engine/Source/Developer/BlueprintCompilerCppBackend/Public/IBlueprintCompilerCppBackendModule.h
+++ b/Engine/Source/Developer/BlueprintCompilerCppBackend/Public/IBlueprintCompilerCppBackendModule.h
@@ -7,8 +7,10 @@
 /**
  * The public interface to this module
  */
-class IBlueprintCompilerCppBackendModule : public IBlueprintCompilerCppBackendModuleInterface
+class BLUEPRINTCOMPILERCPPBACKEND_API IBlueprintCompilerCppBackendModule : public IBlueprintCompilerCppBackendModuleInterface
 {
+public:
 
+	static FString GetBaseFilename(const UObject* AssetObj);
 };
 
diff --git a/Engine/Source/Runtime/CoreUObject/Private/Blueprint/BlueprintSupport.cpp b/Engine/Source/Runtime/CoreUObject/Private/Blueprint/BlueprintSupport.cpp
--- a/Engine/Source/Runtime/CoreUObject/Private/Blueprint/BlueprintSupport.cpp
+++ b/Engine/Source/Runtime/CoreUObject/Private/Blueprint/BlueprintSupport.cpp
@@ -1862,7 +1862,10 @@ bool FReplaceCookedBPGC::CouldBeConverted(const UObject* Object) const
 {
 	if (Object && !Object->HasAnyFlags(RF_ClassDefaultObject) && FConvertibleTypeHelper::IsConvertible(Object))
 	{
-		// TODO: check excluded assets list.
+		if (IsTargetedForConversionDelegate.IsBound())
+		{
+			return IsTargetedForConversionDelegate.Execute(Object);
+		}
 		return true;
 	}
 	return false;
diff --git a/Engine/Source/Runtime/CoreUObject/Public/Blueprint/BlueprintSupport.h b/Engine/Source/Runtime/CoreUObject/Public/Blueprint/BlueprintSupport.h
--- a/Engine/Source/Runtime/CoreUObject/Public/Blueprint/BlueprintSupport.h
+++ b/Engine/Source/Runtime/CoreUObject/Public/Blueprint/BlueprintSupport.h
@@ -157,6 +157,12 @@ public:
 	// The object is of a convertible type, and the object is not listed in excluded assets.
 	bool CouldBeConverted(const UObject* Object) const;
 	bool PackageShouldNotBeSaved(const UPackage* InOuter) const;
+
+	DECLARE_DELEGATE_RetVal_OneParam(bool, FIsTargetedForConversionQuery, const UObject*);
+	FIsTargetedForConversionQuery& OnQueryIfAssetIsTargetedForConversion() { return IsTargetedForConversionDelegate; }
+
+private: 
+	FIsTargetedForConversionQuery IsTargetedForConversionDelegate;
 };
 
 #endif //WITH_EDITOR
