Fixes for cast nodes in cpp backend, to test just generate code for BP_Node_Basic

diff --git a/Engine/Source/Developer/BlueprintCompilerCppBackend/Private/BlueprintCompilerCppBackend.cpp b/Engine/Source/Developer/BlueprintCompilerCppBackend/Private/BlueprintCompilerCppBackend.cpp
--- a/Engine/Source/Developer/BlueprintCompilerCppBackend/Private/BlueprintCompilerCppBackend.cpp
+++ b/Engine/Source/Developer/BlueprintCompilerCppBackend/Private/BlueprintCompilerCppBackend.cpp
@@ -152,7 +152,9 @@ void FBlueprintCompilerCppBackend::EmitCastInterfaceToObjStatement(FEmitterLocal
 
 void FBlueprintCompilerCppBackend::EmitDynamicCastStatement(FEmitterLocalContext& EmitterContext, FKismetFunctionContext& FunctionContext, FBlueprintCompiledStatement& Statement)
 {
-	FString TargetClass = TermToText(EmitterContext, Statement.RHS[0]);
+	UClass* ClassPtr = CastChecked<UClass>(Statement.RHS[0]->ObjectLiteral);
+	check(ClassPtr != nullptr);
+	FString TargetClass = FEmitHelper::GetCppName(ClassPtr);
 	FString ObjectValue = TermToText(EmitterContext, Statement.RHS[1]);
 	FString CastedValue = TermToText(EmitterContext, Statement.LHS);
 	EmitterContext.AddLine(FString::Printf(TEXT("%s = Cast<%s>(%s);"),*CastedValue, *TargetClass, *ObjectValue));
diff --git a/Engine/Source/Editor/BlueprintGraph/Private/DynamicCastHandler.cpp b/Engine/Source/Editor/BlueprintGraph/Private/DynamicCastHandler.cpp
--- a/Engine/Source/Editor/BlueprintGraph/Private/DynamicCastHandler.cpp
+++ b/Engine/Source/Editor/BlueprintGraph/Private/DynamicCastHandler.cpp
@@ -114,6 +114,7 @@ void FKCHandler_DynamicCast::Compile(FKismetFunctionContext& Context, UEdGraphNo
 			return;
 		}
 		CastOpType = KCST_MetaCast;
+		//ClassTerm->Type.PinCategory = CompilerContext.GetSchema()->PC_Class;
 	}
 
 	// Cast Statement
