UEBP-40 fixes - blueprints with multiple delegates to declare now compile, blueprints no logner duplicate delegate declarations, enum defaultss in generated code are now always valid symbols,  minor whitespace improvement to delegate declarations

diff --git a/Engine/Source/Developer/BlueprintCompilerCppBackend/Private/BlueprintCompilerCppBackendBase.cpp b/Engine/Source/Developer/BlueprintCompilerCppBackend/Private/BlueprintCompilerCppBackendBase.cpp
--- a/Engine/Source/Developer/BlueprintCompilerCppBackend/Private/BlueprintCompilerCppBackendBase.cpp
+++ b/Engine/Source/Developer/BlueprintCompilerCppBackend/Private/BlueprintCompilerCppBackendBase.cpp
@@ -61,11 +61,27 @@ void FBlueprintCompilerCppBackendBase::DeclareDelegates(UClass* SourceClass, TIn
 			}
 		}
 
+		// remove duplicates, n^2, but n is small:
+		for (int32 I = 0; I < Delegates.Num(); ++I)
+		{
+			UFunction* TargetFn = Delegates[I]->SignatureFunction;
+			for (int32 J = I + 1; J < Delegates.Num(); ++J)
+			{
+				if (TargetFn == Delegates[J]->SignatureFunction)
+				{
+					// swap erase:
+					Delegates[J] = Delegates[Delegates.Num() - 1];
+					Delegates.RemoveAt(Delegates.Num() - 1);
+				}
+			}
+		}
+
 		auto DelegateDeclarations = FEmitHelper::EmitSinglecastDelegateDeclarations(Delegates);
 		FString AllDeclarations;
 		FEmitHelper::ArrayToString(DelegateDeclarations, AllDeclarations, TEXT(";\n"));
 		if (DelegateDeclarations.Num())
 		{
+			Emit(Header, TEXT("\t"));
 			Emit(Header, *AllDeclarations);
 			Emit(Header, TEXT(";\n"));
 		}
diff --git a/Engine/Source/Developer/BlueprintCompilerCppBackend/Private/BlueprintCompilerCppBackendUtils.cpp b/Engine/Source/Developer/BlueprintCompilerCppBackend/Private/BlueprintCompilerCppBackendUtils.cpp
--- a/Engine/Source/Developer/BlueprintCompilerCppBackend/Private/BlueprintCompilerCppBackendUtils.cpp
+++ b/Engine/Source/Developer/BlueprintCompilerCppBackend/Private/BlueprintCompilerCppBackendUtils.cpp
@@ -164,7 +164,7 @@ void FEmitHelper::ArrayToString(const TArray<FString>& Array, FString& OutString
 	}
 	for (int32 i = 1; i < Array.Num(); ++i)
 	{
-		OutString += TEXT(", ");
+		OutString += Separator;
 		OutString += Array[i];
 	}
 }
@@ -557,8 +557,11 @@ FString FEmitHelper::LiteralTerm(FEmitterLocalContext& EmitterContext, const FEd
 		auto TypeEnum = Cast<UEnum>(Type.PinSubCategoryObject.Get());
 		if (TypeEnum)
 		{
+			// @note: We have to default to the zeroth entry because there may not be a symbol associated with the 
+			// last element (UHT generates a MAX entry for UEnums, even if the user did not declare them, but UHT 
+			// does not generate a symbol for said entry.
 			return FString::Printf(TEXT("%s::%s"), *FEmitHelper::GetCppName(TypeEnum), CustomValue.IsEmpty()
-				? *TypeEnum->GetEnumName(TypeEnum->NumEnums()-1)
+				? *TypeEnum->GetEnumName(0)
 				: *CustomValue);
 		}
 		else
