Basic support for unconverted BP.
- Native UProperties of an unconverted type, are replaced with properties of the first native type.
- support for Cast and DynamicCast
- various fixes

diff --git a/Engine/Source/Developer/BlueprintCompilerCppBackend/Private/BlueprintCompilerCppBackend.cpp b/Engine/Source/Developer/BlueprintCompilerCppBackend/Private/BlueprintCompilerCppBackend.cpp
--- a/Engine/Source/Developer/BlueprintCompilerCppBackend/Private/BlueprintCompilerCppBackend.cpp
+++ b/Engine/Source/Developer/BlueprintCompilerCppBackend/Private/BlueprintCompilerCppBackend.cpp
@@ -154,10 +154,22 @@ void FBlueprintCompilerCppBackend::EmitDynamicCastStatement(FEmitterLocalContext
 {
 	UClass* ClassPtr = CastChecked<UClass>(Statement.RHS[0]->ObjectLiteral);
 	check(ClassPtr != nullptr);
-	FString TargetClass = FEmitHelper::GetCppName(ClassPtr);
-	FString ObjectValue = TermToText(EmitterContext, Statement.RHS[1]);
-	FString CastedValue = TermToText(EmitterContext, Statement.LHS);
-	EmitterContext.AddLine(FString::Printf(TEXT("%s = Cast<%s>(%s);"),*CastedValue, *TargetClass, *ObjectValue));
+	
+	const FString ObjectValue = TermToText(EmitterContext, Statement.RHS[1]);
+	const FString CastedValue = TermToText(EmitterContext, Statement.LHS);
+
+	auto BPGC = Cast<UBlueprintGeneratedClass>(ClassPtr);
+	if (BPGC && !EmitterContext.Dependencies.WillClassBeConverted(BPGC))
+	{
+		const FString NativeClass = FEmitHelper::GetCppName(EmitterContext.GetFirstNativeParent(ClassPtr));
+		const FString TargetClass = EmitterContext.FindGloballyMappedObject(ClassPtr, true);
+		EmitterContext.AddLine(FString::Printf(TEXT("%s = NoNativeCast<%s>(%s, %s);"), *CastedValue, *NativeClass, *TargetClass, *ObjectValue));
+	}
+	else
+	{
+		const FString TargetClass = FEmitHelper::GetCppName(ClassPtr);
+		EmitterContext.AddLine(FString::Printf(TEXT("%s = Cast<%s>(%s);"), *CastedValue, *TargetClass, *ObjectValue));
+	}
 }
 
 void FBlueprintCompilerCppBackend::EmitMetaCastStatement(FEmitterLocalContext& EmitterContext, FKismetFunctionContext& FunctionContext, FBlueprintCompiledStatement& Statement)
@@ -313,13 +325,11 @@ FString FBlueprintCompilerCppBackend::EmitSwitchValueStatmentInner(FEmitterLocal
 		| EPropertyExportCPPFlags::CPPF_NoConst | EPropertyExportCPPFlags::CPPF_NoRef
 		| EPropertyExportCPPFlags::CPPF_BlueprintCppBackend;
 
-	FStringOutputDevice IndexDeclaration;
 	check(IndexTerm && IndexTerm->AssociatedVarProperty);
-	IndexTerm->AssociatedVarProperty->ExportCppDeclaration(IndexDeclaration, EExportedDeclaration::Parameter, nullptr, CppTemplateTypeFlags, true);
+	const FString IndexDeclaration = EmitterContext.ExportCppDeclaration(IndexTerm->AssociatedVarProperty, EExportedDeclaration::Parameter, CppTemplateTypeFlags, true);
 
 	check(DefaultValueTerm && DefaultValueTerm->AssociatedVarProperty);
-	FStringOutputDevice ValueDeclaration;
-	DefaultValueTerm->AssociatedVarProperty->ExportCppDeclaration(ValueDeclaration, EExportedDeclaration::Parameter, nullptr, CppTemplateTypeFlags, true);
+	const FString ValueDeclaration = EmitterContext.ExportCppDeclaration(DefaultValueTerm->AssociatedVarProperty, EExportedDeclaration::Parameter, CppTemplateTypeFlags, true);
 
 	for (int32 TermIndex = TermsBeforeCases; TermIndex < (NumCases * TermsPerCase); TermIndex += TermsPerCase)
 	{
diff --git a/Engine/Source/Developer/BlueprintCompilerCppBackend/Private/BlueprintCompilerCppBackendBase.cpp b/Engine/Source/Developer/BlueprintCompilerCppBackend/Private/BlueprintCompilerCppBackendBase.cpp
--- a/Engine/Source/Developer/BlueprintCompilerCppBackend/Private/BlueprintCompilerCppBackendBase.cpp
+++ b/Engine/Source/Developer/BlueprintCompilerCppBackend/Private/BlueprintCompilerCppBackendBase.cpp
@@ -5,7 +5,7 @@
 #include "BlueprintCompilerCppBackendUtils.h"
 #include "IBlueprintCompilerCppBackendModule.h" // for GetBaseFilename()
 
-void FBlueprintCompilerCppBackendBase::EmitStructProperties(FStringOutputDevice& Target, UStruct* SourceClass)
+void FBlueprintCompilerCppBackendBase::EmitStructProperties(FEmitterLocalContext EmitterContext, UStruct* SourceClass)
 {
 	// Emit class variables
 	for (TFieldIterator<UProperty> It(SourceClass, EFieldIteratorFlags::ExcludeSuper); It; ++It)
@@ -26,16 +26,17 @@ void FBlueprintCompilerCppBackendBase::EmitStructProperties(FStringOutputDevice&
 		}
 		Emit(Header, TEXT(")\n"));
 		Emit(Header, TEXT("\t"));
-		Property->ExportCppDeclaration(Target, EExportedDeclaration::Member, NULL, EPropertyExportCPPFlags::CPPF_CustomTypeName | EPropertyExportCPPFlags::CPPF_BlueprintCppBackend);
+		const FString CppDeclaration = EmitterContext.ExportCppDeclaration(Property, EExportedDeclaration::Member, EPropertyExportCPPFlags::CPPF_CustomTypeName | EPropertyExportCPPFlags::CPPF_BlueprintCppBackend);
+		Emit(Header, *CppDeclaration);
 		Emit(Header, TEXT(";\n"));
 	}
 }
 
-void FBlueprintCompilerCppBackendBase::DeclareDelegates(UClass* SourceClass, TIndirectArray<FKismetFunctionContext>& Functions)
+void FBlueprintCompilerCppBackendBase::DeclareDelegates(FEmitterLocalContext EmitterContext, UClass* SourceClass, TIndirectArray<FKismetFunctionContext>& Functions)
 {
 	// MC DELEGATE DECLARATION
 	{
-		auto DelegateDeclarations = FEmitHelper::EmitMulticastDelegateDeclarations(SourceClass);
+		auto DelegateDeclarations = FEmitHelper::EmitMulticastDelegateDeclarations(EmitterContext, SourceClass);
 		FString AllDeclarations;
 		FEmitHelper::ArrayToString(DelegateDeclarations, AllDeclarations, TEXT(";\n"));
 		if (DelegateDeclarations.Num())
@@ -76,7 +77,7 @@ void FBlueprintCompilerCppBackendBase::DeclareDelegates(UClass* SourceClass, TIn
 			}
 		}
 
-		auto DelegateDeclarations = FEmitHelper::EmitSinglecastDelegateDeclarations(Delegates);
+		auto DelegateDeclarations = FEmitHelper::EmitSinglecastDelegateDeclarations(EmitterContext, Delegates);
 		FString AllDeclarations;
 		FEmitHelper::ArrayToString(DelegateDeclarations, AllDeclarations, TEXT(";\n"));
 		if (DelegateDeclarations.Num())
@@ -95,6 +96,7 @@ void FBlueprintCompilerCppBackendBase::GenerateCodeFromClass(UClass* SourceClass
 	CppClassName = FEmitHelper::GetCppName(SourceClass);
 	
 	FGatherConvertedClassDependencies Dependencies(SourceClass);
+	FEmitterLocalContext EmitterContext(Dependencies);
 	EmitFileBeginning(CleanCppClassName, &Dependencies);
 
 	// Class declaration
@@ -103,8 +105,8 @@ void FBlueprintCompilerCppBackendBase::GenerateCodeFromClass(UClass* SourceClass
 	{
 		Emit(Header, TEXT("UINTERFACE(Blueprintable"));
 		EmitReplaceConvertedMetaData(SourceClass);
-		Emit(Header, *FString::Printf(TEXT(")\nclass U%s : public UInterface\n{\n\tGENERATED_BODY()\n};\n"), *CleanCppClassName));
-		Emit(Header, *FString::Printf(TEXT("\nclass I%s"), *CleanCppClassName));
+		Emit(Header, *FString::Printf(TEXT(")\nclass %s : public UInterface\n{\n\tGENERATED_BODY()\n};\n"), *FEmitHelper::GetCppName(SourceClass, true)));
+		Emit(Header, *FString::Printf(TEXT("\nclass %s"), *CppClassName));
 	}
 	else
 	{
@@ -131,9 +133,9 @@ void FBlueprintCompilerCppBackendBase::GenerateCodeFromClass(UClass* SourceClass
 	// Begin scope
 	Emit(Header,TEXT("\n{\npublic:\n\tGENERATED_BODY()\n"));
 
-	DeclareDelegates(SourceClass, Functions);
+	DeclareDelegates(EmitterContext, SourceClass, Functions);
 
-	EmitStructProperties(Header, SourceClass);
+	EmitStructProperties(EmitterContext, SourceClass);
 
 	// Create the state map
 	for (int32 i = 0; i < Functions.Num(); ++i)
@@ -148,8 +150,6 @@ void FBlueprintCompilerCppBackendBase::GenerateCodeFromClass(UClass* SourceClass
 		Emit(Header, TEXT("\n"));
 	}
 
-	FEmitterLocalContext EmitterContext(SourceClass, Dependencies);
-
 	if (!bIsInterface)
 	{
 		Emit(Header, *FString::Printf(TEXT("\t%s(const FObjectInitializer& ObjectInitializer = FObjectInitializer::Get());\n\n"), *CppClassName));
@@ -174,14 +174,15 @@ void FBlueprintCompilerCppBackendBase::GenerateCodeFromClass(UClass* SourceClass
 	Emit(Body, *FEmitHelper::EmitLifetimeReplicatedPropsImpl(SourceClass, CppClassName, TEXT("")));
 }
 
-void FBlueprintCompilerCppBackendBase::DeclareLocalVariables(FKismetFunctionContext& FunctionContext, TArray<UProperty*>& LocalVariables)
+void FBlueprintCompilerCppBackendBase::DeclareLocalVariables(FKismetFunctionContext& FunctionContext, FEmitterLocalContext& EmitterContext, TArray<UProperty*>& LocalVariables)
 {
 	for (int32 i = 0; i < LocalVariables.Num(); ++i)
 	{
 		UProperty* LocalVariable = LocalVariables[i];
 
 		Emit(Body, TEXT("\t"));
-		LocalVariable->ExportCppDeclaration(Body, EExportedDeclaration::Local, NULL, EPropertyExportCPPFlags::CPPF_CustomTypeName | EPropertyExportCPPFlags::CPPF_BlueprintCppBackend);
+		const FString CppDeclaration = EmitterContext.ExportCppDeclaration(LocalVariable, EExportedDeclaration::Local, EPropertyExportCPPFlags::CPPF_CustomTypeName | EPropertyExportCPPFlags::CPPF_BlueprintCppBackend);
+		Emit(Body, *CppDeclaration);
 		Emit(Body, TEXT("{};\n"));
 	}
 
@@ -252,7 +253,20 @@ void FBlueprintCompilerCppBackendBase::ConstructFunction(FKismetFunctionContext&
 		}
 
 		// Emit the declaration
-		const FString ReturnType = ReturnValue ? ReturnValue->GetCPPType(NULL, EPropertyExportCPPFlags::CPPF_CustomTypeName) : TEXT("void");
+		FString ReturnType;
+		if (ReturnValue)
+		{
+			const uint32 LocalExportCPPFlags = EPropertyExportCPPFlags::CPPF_CustomTypeName
+				| EPropertyExportCPPFlags::CPPF_NoConst 
+				| EPropertyExportCPPFlags::CPPF_NoRef 
+				| EPropertyExportCPPFlags::CPPF_NoStaticArray 
+				| EPropertyExportCPPFlags::CPPF_BlueprintCppBackend;
+			ReturnType = EmitterContext.ExportCppDeclaration(ReturnValue, EExportedDeclaration::Parameter, LocalExportCPPFlags, true);
+		}
+		else
+		{
+			ReturnType = TEXT("void");
+		}
 
 		//@TODO: Make the header+body export more uniform
 		{
@@ -299,8 +313,9 @@ void FBlueprintCompilerCppBackendBase::ConstructFunction(FKismetFunctionContext&
 					Emit(Body, TEXT("/*out*/ "));
 				}
 
-				ArgProperty->ExportCppDeclaration(Header, EExportedDeclaration::Parameter, NULL, EPropertyExportCPPFlags::CPPF_CustomTypeName | EPropertyExportCPPFlags::CPPF_BlueprintCppBackend);
-				ArgProperty->ExportCppDeclaration(Body, EExportedDeclaration::Parameter, NULL, EPropertyExportCPPFlags::CPPF_CustomTypeName | EPropertyExportCPPFlags::CPPF_BlueprintCppBackend);
+				const FString TypeDeclaration = EmitterContext.ExportCppDeclaration(ArgProperty, EExportedDeclaration::Parameter, EPropertyExportCPPFlags::CPPF_CustomTypeName | EPropertyExportCPPFlags::CPPF_BlueprintCppBackend);
+				Emit(Header, *TypeDeclaration);
+				Emit(Body, *TypeDeclaration);
 			}
 
 			Emit(Header, TEXT(")"));
@@ -319,7 +334,7 @@ void FBlueprintCompilerCppBackendBase::ConstructFunction(FKismetFunctionContext&
 	if (!bGenerateStubOnly)
 	{
 		// Emit local variables
-		DeclareLocalVariables(FunctionContext, LocalVariables);
+		DeclareLocalVariables(FunctionContext, EmitterContext, LocalVariables);
 
 		const bool bUseSwitchState = FunctionContext.MustUseSwitchState(nullptr);
 
@@ -395,6 +410,7 @@ void FBlueprintCompilerCppBackendBase::GenerateCodeFromStruct(UUserDefinedStruct
 	check(SourceStruct);
 
 	FGatherConvertedClassDependencies Dependencies(SourceStruct);
+	FEmitterLocalContext EmitterContext(Dependencies);
 	// use GetBaseFilename() so that we can coordinate #includes and filenames
 	EmitFileBeginning(FEmitHelper::GetBaseFilename(SourceStruct), &Dependencies);
 
@@ -404,7 +420,7 @@ void FBlueprintCompilerCppBackendBase::GenerateCodeFromStruct(UUserDefinedStruct
 	Emit(Header, TEXT(")\n"));
 
 	Emit(Header, *FString::Printf(TEXT("struct %s\n{\npublic:\n\tGENERATED_BODY()\n"), *NewName));
-	EmitStructProperties(Header, SourceStruct);
+	EmitStructProperties(EmitterContext, SourceStruct);
 	Emit(Header, *FEmitDefaultValueHelper::GenerateGetDefaultValue(SourceStruct, Dependencies));
 	Emit(Header, TEXT("};\n"));
 }
diff --git a/Engine/Source/Developer/BlueprintCompilerCppBackend/Private/BlueprintCompilerCppBackendBase.h b/Engine/Source/Developer/BlueprintCompilerCppBackend/Private/BlueprintCompilerCppBackendBase.h
--- a/Engine/Source/Developer/BlueprintCompilerCppBackend/Private/BlueprintCompilerCppBackendBase.h
+++ b/Engine/Source/Developer/BlueprintCompilerCppBackend/Private/BlueprintCompilerCppBackendBase.h
@@ -74,10 +74,10 @@ protected:
 	*/
 	virtual FString InnerFunctionImplementation(FKismetFunctionContext& FunctionContext, FEmitterLocalContext& EmitterContext, bool bUseSwitchState) PURE_VIRTUAL(FBlueprintCompilerCppBackendBase::InnerFunctionImplementation, return FString(););
 
-	void EmitStructProperties(FStringOutputDevice& Target, UStruct* SourceClass);
+	void EmitStructProperties(FEmitterLocalContext EmitterContext, UStruct* SourceClass);
 
 	/** Emits local variable declarations for a function */
-	void DeclareLocalVariables(FKismetFunctionContext& FunctionContext, TArray<UProperty*>& LocalVariables);
+	void DeclareLocalVariables(FKismetFunctionContext& FunctionContext, FEmitterLocalContext& EmitterContext, TArray<UProperty*>& LocalVariables);
 	
 	/** Builds both the header declaration and body implementation of a function */
 	void ConstructFunction(FKismetFunctionContext& FunctionContext, FEmitterLocalContext& EmitterContext, bool bGenerateStubOnly);
@@ -92,5 +92,5 @@ protected:
 
 	void EmitReplaceConvertedMetaData(UObject* Obj);
 
-	void DeclareDelegates(UClass* SourceClass, TIndirectArray<FKismetFunctionContext>& Functions);
+	void DeclareDelegates(FEmitterLocalContext EmitterContext, UClass* SourceClass, TIndirectArray<FKismetFunctionContext>& Functions);
 };
diff --git a/Engine/Source/Developer/BlueprintCompilerCppBackend/Private/BlueprintCompilerCppBackendGatherDependencies.cpp b/Engine/Source/Developer/BlueprintCompilerCppBackend/Private/BlueprintCompilerCppBackendGatherDependencies.cpp
--- a/Engine/Source/Developer/BlueprintCompilerCppBackend/Private/BlueprintCompilerCppBackendGatherDependencies.cpp
+++ b/Engine/Source/Developer/BlueprintCompilerCppBackend/Private/BlueprintCompilerCppBackendGatherDependencies.cpp
@@ -121,7 +121,7 @@ void FGatherConvertedClassDependencies::HandleObjectReference(UObject*& InObject
 			ConvertedEnum.Add(CastChecked<UUserDefinedEnum>(InObject));
 		}
 	}
-	else if (InObject->IsAsset() && !InObject->IsIn(OriginalStruct))
+	else if ((InObject->IsAsset() || ObjAsBPGC) && !InObject->IsIn(OriginalStruct))
 	{
 		Assets.Add(InObject);
 		return;
diff --git a/Engine/Source/Developer/BlueprintCompilerCppBackend/Private/BlueprintCompilerCppBackendUMG.cpp b/Engine/Source/Developer/BlueprintCompilerCppBackend/Private/BlueprintCompilerCppBackendUMG.cpp
--- a/Engine/Source/Developer/BlueprintCompilerCppBackend/Private/BlueprintCompilerCppBackendUMG.cpp
+++ b/Engine/Source/Developer/BlueprintCompilerCppBackend/Private/BlueprintCompilerCppBackendUMG.cpp
@@ -51,10 +51,10 @@ void FBackendHelperUMG::EmitWidgetInitializationFunctions(FEmitterLocalContext&
 		{
 			check(InProperty && DataPtr);
 			const FString NativeName = InContext.GenerateUniqueLocalName();
-			FStringOutputDevice Target;
-			InProperty->ExportCppDeclaration(Target, EExportedDeclaration::Local, nullptr
-				, EPropertyExportCPPFlags::CPPF_CustomTypeName | EPropertyExportCPPFlags::CPPF_BlueprintCppBackend | EPropertyExportCPPFlags::CPPF_NoConst | EPropertyExportCPPFlags::CPPF_NoRef
-				, true);
+			
+			const uint32 CppTemplateTypeFlags = EPropertyExportCPPFlags::CPPF_CustomTypeName | EPropertyExportCPPFlags::CPPF_BlueprintCppBackend | EPropertyExportCPPFlags::CPPF_NoConst | EPropertyExportCPPFlags::CPPF_NoRef;
+			const FString Target = InContext.ExportCppDeclaration(InProperty, EExportedDeclaration::Local, CppTemplateTypeFlags, true);
+
 			InContext.AddLine(FString::Printf(TEXT("%s %s;"), *Target, *NativeName));
 			FEmitDefaultValueHelper::InnerGenerate(InContext, InProperty, NativeName, DataPtr, nullptr, true);
 			return NativeName;
diff --git a/Engine/Source/Developer/BlueprintCompilerCppBackend/Private/BlueprintCompilerCppBackendUtils.cpp b/Engine/Source/Developer/BlueprintCompilerCppBackend/Private/BlueprintCompilerCppBackendUtils.cpp
--- a/Engine/Source/Developer/BlueprintCompilerCppBackend/Private/BlueprintCompilerCppBackendUtils.cpp
+++ b/Engine/Source/Developer/BlueprintCompilerCppBackend/Private/BlueprintCompilerCppBackendUtils.cpp
@@ -25,6 +25,7 @@ FString FEmitterLocalContext::GenerateUniqueLocalName()
 
 FString FEmitterLocalContext::FindGloballyMappedObject(UObject* Object, bool bLoadIfNotFound)
 {
+	UClass* ActualClass = Cast<UClass>(Dependencies.GetOriginalStruct());
 	if (ActualClass && Object && Object->IsIn(ActualClass))
 	{
 		if (const FString* NamePtr = (CurrentCodeType == EGeneratedCodeType::SubobjectsOfClass) ? ClassSubobjectsMap.Find(Object) : nullptr)
@@ -119,6 +120,117 @@ FString FEmitterLocalContext::FindGloballyMappedObject(UObject* Object, bool bLo
 	return FString{};
 }
 
+FString FEmitterLocalContext::ExportTextItem(const UProperty* Property, const void* PropertyValue) const
+{
+	if (auto ArrayProperty = Cast<const UArrayProperty>(Property))
+	{
+		const uint32 LocalExportCPPFlags = EPropertyExportCPPFlags::CPPF_CustomTypeName
+			| EPropertyExportCPPFlags::CPPF_NoConst
+			| EPropertyExportCPPFlags::CPPF_NoRef
+			| EPropertyExportCPPFlags::CPPF_NoStaticArray
+			| EPropertyExportCPPFlags::CPPF_BlueprintCppBackend;
+		const FString TypeText = ExportCppDeclaration(ArrayProperty, EExportedDeclaration::Parameter, LocalExportCPPFlags, true);
+		return FString::Printf(TEXT("%s()"), *TypeText);
+	}
+
+	auto ObjectPropertyBase = Cast<const UObjectPropertyBase>(Property);
+	if (ObjectPropertyBase && Property->IsA<UAssetObjectProperty>())
+	{
+		if (UObject* Object = ObjectPropertyBase->GetObjectPropertyValue(PropertyValue))
+		{
+			UClass* ActualClass = ObjectPropertyBase->PropertyClass;
+			auto BPGC = Cast<UBlueprintGeneratedClass>(ActualClass);
+			const bool bConverted = BPGC && Dependencies.WillClassBeConverted(BPGC);
+			if (BPGC && !bConverted)
+			{
+				ActualClass = GetFirstNativeParent(BPGC);
+			}
+			return FString::Printf(TEXT("LoadObject<%s>(nullptr, TEXT(\"%s\"))")
+				, *FEmitHelper::GetCppName(ActualClass)
+				, *(Object->GetPathName().ReplaceCharWithEscapedChar()));
+		}
+		return TEXT("nullptr");
+	}
+
+	FString ValueStr;
+	Property->ExportTextItem(ValueStr, PropertyValue, PropertyValue, nullptr, EPropertyPortFlags::PPF_ExportCpp);
+	return ValueStr;
+}
+
+FString FEmitterLocalContext::ExportCppDeclaration(const UProperty* Property, EExportedDeclaration::Type DeclarationType, uint32 InExportCPPFlags, bool bSkipParameterName) const
+{
+	FString ActualCppType;
+	FString* ActualCppTypePtr = nullptr;
+	FString ActualExtendedType;
+	FString* ActualExtendedTypePtr = nullptr;
+	uint32 ExportCPPFlags = InExportCPPFlags;
+
+	auto GetActualNameCPP = [&](const UObjectPropertyBase* ObjectPropertyBase, UClass* InActualClass)
+	{
+		auto BPGC = Cast<UBlueprintGeneratedClass>(InActualClass);
+		const bool bConverted = BPGC && Dependencies.WillClassBeConverted(BPGC);
+		if (BPGC && !bConverted)
+		{
+			const bool bIsParameter = (DeclarationType == EExportedDeclaration::Parameter) || (DeclarationType == EExportedDeclaration::MacroParameter);
+			const uint32 LocalExportCPPFlags = ExportCPPFlags | (bIsParameter ? CPPF_ArgumentOrReturnValue : 0);
+			UClass* NativeType = GetFirstNativeParent(BPGC);
+			ActualCppType = ObjectPropertyBase->GetCPPTypeCustom(&ActualExtendedType, LocalExportCPPFlags, NativeType);
+			ActualCppTypePtr = &ActualCppType;
+			if (!ActualExtendedType.IsEmpty())
+			{
+				ActualExtendedTypePtr = &ActualExtendedType;
+			}
+		}
+	};
+
+	auto ArrayProperty = Cast<const UArrayProperty>(Property);
+	if (ArrayProperty)
+	{
+		Property = ArrayProperty->Inner;
+		ExportCPPFlags = (ExportCPPFlags & ~CPPF_ArgumentOrReturnValue);
+	}
+
+	if (auto ClassProperty = Cast<const UClassProperty>(Property))
+	{
+		GetActualNameCPP(ClassProperty, ClassProperty->MetaClass);
+	}
+	else if (auto AssetClassProperty = Cast<const UAssetClassProperty>(Property))
+	{
+		GetActualNameCPP(AssetClassProperty, AssetClassProperty->MetaClass);
+	}
+	else if (auto ObjectProperty = Cast<const UObjectPropertyBase>(Property))
+	{
+		GetActualNameCPP(ObjectProperty, ObjectProperty->PropertyClass);
+	}
+
+	if (ArrayProperty)
+	{
+		Property = ArrayProperty;
+		if (ActualCppTypePtr)
+		{
+			const FString LocalActualCppType = ActualCppType;
+			ActualCppType.Empty();
+			const FString LocalActualExtendedType = ActualExtendedType;
+			ActualExtendedType.Empty();
+			ActualExtendedTypePtr = nullptr;
+
+			const bool bIsParameter = (DeclarationType == EExportedDeclaration::Parameter) || (DeclarationType == EExportedDeclaration::MacroParameter);
+			const uint32 LocalExportCPPFlags = InExportCPPFlags | (bIsParameter ? CPPF_ArgumentOrReturnValue : 0);
+
+			ActualCppType = ArrayProperty->GetCPPTypeCustom(&ActualExtendedType, LocalExportCPPFlags, LocalActualCppType, LocalActualExtendedType);
+			if (!ActualExtendedType.IsEmpty())
+			{
+				ActualExtendedTypePtr = &ActualExtendedType;
+			}
+		}
+	}
+
+	FStringOutputDevice Out;
+	Property->ExportCppDeclaration(Out, DeclarationType, nullptr, InExportCPPFlags, bSkipParameterName, ActualCppTypePtr, ActualExtendedTypePtr);
+	return FString(Out);
+
+}
+
 FString FEmitHelper::GetCppName(const UField* Field, bool bUInterface)
 {
 	check(Field);
@@ -438,14 +550,14 @@ FString FEmitHelper::EmitUFuntion(UFunction* Function, TArray<FString>* Additina
 	return FString::Printf(TEXT("UFUNCTION(%s)"), *AllTags);
 }
 
-int32 FEmitHelper::ParseDelegateDetails(UFunction* Signature, FString& OutParametersMacro, FString& OutParamNumberStr)
+int32 FEmitHelper::ParseDelegateDetails(FEmitterLocalContext& EmitterContext, UFunction* Signature, FString& OutParametersMacro, FString& OutParamNumberStr)
 {
 	int32 ParameterNum = 0;
 	FStringOutputDevice Parameters;
 	for (TFieldIterator<UProperty> PropIt(Signature); PropIt && (PropIt->PropertyFlags & CPF_Parm); ++PropIt)
 	{
 		Parameters += ", ";
-		PropIt->ExportCppDeclaration(Parameters, EExportedDeclaration::MacroParameter, NULL, EPropertyExportCPPFlags::CPPF_CustomTypeName);
+		Parameters += EmitterContext.ExportCppDeclaration(*PropIt, EExportedDeclaration::MacroParameter, EPropertyExportCPPFlags::CPPF_CustomTypeName | EPropertyExportCPPFlags::CPPF_BlueprintCppBackend);
 		++ParameterNum;
 	}
 
@@ -469,7 +581,7 @@ int32 FEmitHelper::ParseDelegateDetails(UFunction* Signature, FString& OutParame
 	return ParameterNum;
 }
 
-TArray<FString> FEmitHelper::EmitSinglecastDelegateDeclarations(const TArray<UDelegateProperty*>& Delegates)
+TArray<FString> FEmitHelper::EmitSinglecastDelegateDeclarations(FEmitterLocalContext& EmitterContext, const TArray<UDelegateProperty*>& Delegates)
 {
 	TArray<FString> Results;
 	for (auto It : Delegates)
@@ -479,15 +591,21 @@ TArray<FString> FEmitHelper::EmitSinglecastDelegateDeclarations(const TArray<UDe
 		check(Signature);
 
 		FString ParamNumberStr, Parameters;
-		ParseDelegateDetails(Signature, Parameters, ParamNumberStr);
+		ParseDelegateDetails(EmitterContext, Signature, Parameters, ParamNumberStr);
 
-		Results.Add(*FString::Printf(TEXT("DECLARE_DYNAMIC_DELEGATE%s(%s%s)"),
-			*ParamNumberStr, *It->GetCPPType(NULL, EPropertyExportCPPFlags::CPPF_CustomTypeName), *Parameters));
+		const uint32 LocalExportCPPFlags = EPropertyExportCPPFlags::CPPF_CustomTypeName
+			| EPropertyExportCPPFlags::CPPF_NoConst
+			| EPropertyExportCPPFlags::CPPF_NoRef
+			| EPropertyExportCPPFlags::CPPF_NoStaticArray
+			| EPropertyExportCPPFlags::CPPF_BlueprintCppBackend;
+		const FString TypeName = EmitterContext.ExportCppDeclaration(It, EExportedDeclaration::Parameter, LocalExportCPPFlags, true);
+
+		Results.Add(*FString::Printf(TEXT("DECLARE_DYNAMIC_DELEGATE%s(%s%s)"), *ParamNumberStr, *TypeName, *Parameters));
 	}
 	return Results;
 }
 
-TArray<FString> FEmitHelper::EmitMulticastDelegateDeclarations(UClass* SourceClass)
+TArray<FString> FEmitHelper::EmitMulticastDelegateDeclarations(FEmitterLocalContext& EmitterContext, UClass* SourceClass)
 {
 	TArray<FString> Results;
 	for (TFieldIterator<UMulticastDelegateProperty> It(SourceClass, EFieldIteratorFlags::ExcludeSuper); It; ++It)
@@ -496,10 +614,16 @@ TArray<FString> FEmitHelper::EmitMulticastDelegateDeclarations(UClass* SourceCla
 		check(Signature);
 
 		FString ParamNumberStr, Parameters;
-		ParseDelegateDetails(Signature, Parameters, ParamNumberStr);
+		ParseDelegateDetails(EmitterContext, Signature, Parameters, ParamNumberStr);
+
+		const uint32 LocalExportCPPFlags = EPropertyExportCPPFlags::CPPF_CustomTypeName
+			| EPropertyExportCPPFlags::CPPF_NoConst
+			| EPropertyExportCPPFlags::CPPF_NoRef
+			| EPropertyExportCPPFlags::CPPF_NoStaticArray
+			| EPropertyExportCPPFlags::CPPF_BlueprintCppBackend;
+		const FString TypeName = EmitterContext.ExportCppDeclaration(*It, EExportedDeclaration::Parameter, LocalExportCPPFlags, true);
 
-		Results.Add(*FString::Printf(TEXT("\tDECLARE_DYNAMIC_MULTICAST_DELEGATE%s(%s%s)"),
-			*ParamNumberStr, *It->GetCPPType(NULL, EPropertyExportCPPFlags::CPPF_CustomTypeName), *Parameters));
+		Results.Add(*FString::Printf(TEXT("\tDECLARE_DYNAMIC_MULTICAST_DELEGATE%s(%s%s)"), *ParamNumberStr, *TypeName, *Parameters));
 	}
 
 	return Results;
@@ -668,7 +792,7 @@ FString FEmitHelper::LiteralTerm(FEmitterLocalContext& EmitterContext, const FEd
 			{
 				return MappedObject;
 			}
-			return FString::Printf(TEXT("LoadClass<UObject>(nullptr, TEXT(\"%s\"), nullptr, 0, nullptr "), *(LiteralObject->GetPathName().ReplaceCharWithEscapedChar()));
+			return FString::Printf(TEXT("LoadClass<UObject>(nullptr, TEXT(\"%s\"), nullptr, 0, nullptr)"), *(LiteralObject->GetPathName().ReplaceCharWithEscapedChar()));
 		}
 		return FString(TEXT("nullptr"));
 	}
diff --git a/Engine/Source/Developer/BlueprintCompilerCppBackend/Private/BlueprintCompilerCppBackendUtils.h b/Engine/Source/Developer/BlueprintCompilerCppBackend/Private/BlueprintCompilerCppBackendUtils.h
--- a/Engine/Source/Developer/BlueprintCompilerCppBackend/Private/BlueprintCompilerCppBackendUtils.h
+++ b/Engine/Source/Developer/BlueprintCompilerCppBackend/Private/BlueprintCompilerCppBackendUtils.h
@@ -29,7 +29,6 @@ private:
 	int32 LocalNameIndexMax;
 	TArray<FString> HighPriorityLines;
 	TArray<FString> LowPriorityLines;
-	UClass* ActualClass;
 
 	//ConstructorOnly Local Names
 	TMap<UObject*, FString> ClassSubobjectsMap;
@@ -45,18 +44,12 @@ public:
 
 	const FGatherConvertedClassDependencies& Dependencies;
 
-	FEmitterLocalContext(UClass* InClass, const FGatherConvertedClassDependencies& InDependencies)
+	FEmitterLocalContext(const FGatherConvertedClassDependencies& InDependencies)
 		: CurrentCodeType(EGeneratedCodeType::Regular)
 		, LocalNameIndexMax(0)
-		, ActualClass(InClass)
 		, Dependencies(InDependencies)
 	{}
 
-	UClass* GetActualClass() const
-	{
-		return ActualClass;
-	}
-
 	// CONSTRUCTOR FUNCTIONS
 
 	const TCHAR* ClassSubobjectListName(EClassSubobjectList ListType)
@@ -110,13 +103,28 @@ public:
 
 	// UNIVERSAL FUNCTIONS
 
+	UClass* GetFirstNativeParent(UClass* InClass) const
+	{
+		check(InClass);
+		for (UClass* ItClass = InClass; ItClass; ItClass = ItClass->GetSuperClass())
+		{
+			auto BPGC = Cast<UBlueprintGeneratedClass>(ItClass);
+			if (ItClass->HasAnyClassFlags(CLASS_Native) || (BPGC && Dependencies.WillClassBeConverted(BPGC)))
+			{
+				return ItClass;
+			}
+		}
+		check(false);
+		return nullptr;
+	}
+
 	FString GenerateGetProperty(const UProperty* Property) const;
 
 	FString GenerateUniqueLocalName();
 
 	UClass* GetCurrentlyGeneratedClass() const
 	{
-		return ActualClass;
+		return Cast<UClass>(Dependencies.GetOriginalStruct());
 	}
 
 	/** All objects (that can be referenced from other package) that will have a different path in cooked build
@@ -171,6 +179,11 @@ public:
 		FlushLines();
 		Result.Reset();
 	}
+
+	// Functions needed for Unconverted classes
+	FString ExportCppDeclaration(const UProperty* Property, EExportedDeclaration::Type DeclarationType, uint32 AdditionalExportCPPFlags, bool bSkipParameterName = false) const;
+	FString ExportTextItem(const UProperty* Property, const void* PropertyValue) const;
+
 };
 
 struct FEmitHelper
@@ -198,11 +211,11 @@ struct FEmitHelper
 
 	static FString EmitUFuntion(UFunction* Function, TArray<FString>* AdditinalMetaData = nullptr);
 
-	static int32 ParseDelegateDetails(UFunction* Signature, FString& OutParametersMacro, FString& OutParamNumberStr);
+	static int32 ParseDelegateDetails(FEmitterLocalContext& EmitterContext, UFunction* Signature, FString& OutParametersMacro, FString& OutParamNumberStr);
 
-	static TArray<FString> EmitSinglecastDelegateDeclarations(const TArray<UDelegateProperty*>& Delegates);
+	static TArray<FString> EmitSinglecastDelegateDeclarations(FEmitterLocalContext& EmitterContext, const TArray<UDelegateProperty*>& Delegates);
 
-	static TArray<FString> EmitMulticastDelegateDeclarations(UClass* SourceClass);
+	static TArray<FString> EmitMulticastDelegateDeclarations(FEmitterLocalContext& EmitterContext, UClass* SourceClass);
 
 	static FString EmitLifetimeReplicatedPropsImpl(UClass* SourceClass, const FString& CppClassName, const TCHAR* InCurrentIndent);
 
diff --git a/Engine/Source/Developer/BlueprintCompilerCppBackend/Private/BlueprintCompilerCppBackendValueHelper.cpp b/Engine/Source/Developer/BlueprintCompilerCppBackend/Private/BlueprintCompilerCppBackendValueHelper.cpp
--- a/Engine/Source/Developer/BlueprintCompilerCppBackend/Private/BlueprintCompilerCppBackendValueHelper.cpp
+++ b/Engine/Source/Developer/BlueprintCompilerCppBackend/Private/BlueprintCompilerCppBackendValueHelper.cpp
@@ -51,11 +51,10 @@ void FEmitDefaultValueHelper::OuterGenerate(FEmitterLocalContext& Context
 				
 				const FString PropertyObject = Context.GenerateGetProperty(Property); //X::StaticClass()->FindPropertyByName(GET_MEMBER_NAME_CHECKED(X, X))
 
-				FStringOutputDevice TypeDeclaration;
 				const uint32 CppTemplateTypeFlags = EPropertyExportCPPFlags::CPPF_CustomTypeName
 					| EPropertyExportCPPFlags::CPPF_NoConst | EPropertyExportCPPFlags::CPPF_NoRef | EPropertyExportCPPFlags::CPPF_NoStaticArray
 					| EPropertyExportCPPFlags::CPPF_BlueprintCppBackend;
-				Property->ExportCppDeclaration(TypeDeclaration, EExportedDeclaration::Parameter, nullptr, CppTemplateTypeFlags, true);
+				const FString TypeDeclaration = Context.ExportCppDeclaration(Property, EExportedDeclaration::Parameter, CppTemplateTypeFlags, true);
 
 				const FString OperatorStr = (EPropertyAccessOperator::Dot == AccessOperator) ? TEXT("&") : TEXT("");
 				const FString ContainerStr = (EPropertyAccessOperator::None == AccessOperator) ? TEXT("this") : FString::Printf(TEXT("%s(%s)"), *OperatorStr, *OuterPath);
@@ -87,7 +86,7 @@ FString FEmitDefaultValueHelper::GenerateGetDefaultValue(const UUserDefinedStruc
 	FStructOnScope StructData(Struct);
 	FStructureEditorUtils::Fill_MakeStructureDefaultValue(Struct, StructData.GetStructMemory());
 
-	FEmitterLocalContext Context(nullptr, Dependencies);
+	FEmitterLocalContext Context(Dependencies);
 	Context.IncreaseIndent();
 	Context.IncreaseIndent();
 	for (auto Property : TFieldRange<const UProperty>(Struct))
@@ -107,8 +106,8 @@ void FEmitDefaultValueHelper::InnerGenerate(FEmitterLocalContext& Context, const
 		FString ValueStr = HandleSpecialTypes(LocalContext, LocalProperty, LocalValuePtr);
 		if (ValueStr.IsEmpty())
 		{
+			ValueStr = LocalContext.ExportTextItem(LocalProperty, LocalValuePtr);
 			auto StructProperty = Cast<const UStructProperty>(LocalProperty);
-			LocalProperty->ExportTextItem(ValueStr, LocalValuePtr, LocalValuePtr, nullptr, EPropertyPortFlags::PPF_ExportCpp);
 			if (ValueStr.IsEmpty() && StructProperty)
 			{
 				check(StructProperty->Struct);
@@ -440,7 +439,7 @@ struct FDependenciesHelper
 
 FString FEmitDefaultValueHelper::GenerateConstructor(FEmitterLocalContext& Context)
 {
-	auto BPGC = CastChecked<UBlueprintGeneratedClass>(Context.GetActualClass());
+	auto BPGC = CastChecked<UBlueprintGeneratedClass>(Context.GetCurrentlyGeneratedClass());
 	const FString CppClassName = FEmitHelper::GetCppName(BPGC);
 
 	Context.AddLine(FString::Printf(TEXT("%s::%s(const FObjectInitializer& ObjectInitializer) : Super(ObjectInitializer)"), *CppClassName, *CppClassName));
diff --git a/Engine/Source/Runtime/CoreUObject/Private/UObject/Property.cpp b/Engine/Source/Runtime/CoreUObject/Private/UObject/Property.cpp
--- a/Engine/Source/Runtime/CoreUObject/Private/UObject/Property.cpp
+++ b/Engine/Source/Runtime/CoreUObject/Private/UObject/Property.cpp
@@ -452,7 +452,8 @@ FString UProperty::GetCPPMacroType( FString& ExtendedTypeText ) const
 	return TEXT("PROPERTY");
 }
 
-void UProperty::ExportCppDeclaration(FOutputDevice& Out, EExportedDeclaration::Type DeclarationType, const TCHAR* ArrayDimOverride, uint32 AdditionalExportCPPFlags, bool bSkipParameterName) const
+void UProperty::ExportCppDeclaration(FOutputDevice& Out, EExportedDeclaration::Type DeclarationType, const TCHAR* ArrayDimOverride, uint32 AdditionalExportCPPFlags
+	, bool bSkipParameterName, const FString* ActualCppType, const FString* ActualExtendedType) const
 {
 	const bool bIsParameter = (DeclarationType == EExportedDeclaration::Parameter) || (DeclarationType == EExportedDeclaration::MacroParameter);
 	const bool bIsInterfaceProp = dynamic_cast<const UInterfaceProperty*>(this) != nullptr;
@@ -460,7 +461,20 @@ void UProperty::ExportCppDeclaration(FOutputDevice& Out, EExportedDeclaration::T
 	// export the property type text (e.g. FString; int32; TArray, etc.)
 	FString ExtendedTypeText;
 	const uint32 ExportCPPFlags = AdditionalExportCPPFlags | (bIsParameter ? CPPF_ArgumentOrReturnValue : 0);
-	FString TypeText = GetCPPType(&ExtendedTypeText, ExportCPPFlags);
+	FString TypeText;
+	if (ActualCppType)
+	{
+		TypeText = *ActualCppType;
+	}
+	else
+	{
+		TypeText = GetCPPType(&ExtendedTypeText, ExportCPPFlags);
+	}
+
+	if (ActualExtendedType)
+	{
+		ExtendedTypeText = *ActualExtendedType;
+	}
 
 	const bool bCanHaveConst = 0 == (AdditionalExportCPPFlags & CPPF_NoConst);
 	if (!dynamic_cast<const UBoolProperty*>(this) && bCanHaveConst) // can't have const bitfields because then we cannot determine their offset and mask from the compiler
diff --git a/Engine/Source/Runtime/CoreUObject/Private/UObject/PropertyArray.cpp b/Engine/Source/Runtime/CoreUObject/Private/UObject/PropertyArray.cpp
--- a/Engine/Source/Runtime/CoreUObject/Private/UObject/PropertyArray.cpp
+++ b/Engine/Source/Runtime/CoreUObject/Private/UObject/PropertyArray.cpp
@@ -87,20 +87,18 @@ void UArrayProperty::AddReferencedObjects(UObject* InThis, FReferenceCollector&
 	Collector.AddReferencedObject( This->Inner, This );
 	Super::AddReferencedObjects( This, Collector );
 }
-FString UArrayProperty::GetCPPType( FString* ExtendedTypeText/*=NULL*/, uint32 CPPExportFlags/*=0*/ ) const
-{
-	checkSlow(Inner);
 
-	if ( ExtendedTypeText != NULL )
+FString UArrayProperty::GetCPPTypeCustom(FString* ExtendedTypeText, uint32 CPPExportFlags, const FString& InnerTypeText, const FString& InInnerExtendedTypeText) const
+{
+	if (ExtendedTypeText != NULL)
 	{
-		FString InnerExtendedTypeText;
-		FString InnerTypeText = Inner->GetCPPType(&InnerExtendedTypeText, CPPExportFlags & ~CPPF_ArgumentOrReturnValue); // we won't consider array inners to be "arguments or return values"
-		if ( InnerExtendedTypeText.Len() && InnerExtendedTypeText.Right(1) == TEXT(">") )
+		FString InnerExtendedTypeText = InInnerExtendedTypeText;
+		if (InnerExtendedTypeText.Len() && InnerExtendedTypeText.Right(1) == TEXT(">"))
 		{
 			// if our internal property type is a template class, add a space between the closing brackets b/c VS.NET cannot parse this correctly
 			InnerExtendedTypeText += TEXT(" ");
 		}
-		else if ( !InnerExtendedTypeText.Len() && InnerTypeText.Len() && InnerTypeText.Right(1) == TEXT(">") )
+		else if (!InnerExtendedTypeText.Len() && InnerTypeText.Len() && InnerTypeText.Right(1) == TEXT(">"))
 		{
 			// if our internal property type is a template class, add a space between the closing brackets b/c VS.NET cannot parse this correctly
 			InnerExtendedTypeText += TEXT(" ");
@@ -110,6 +108,18 @@ FString UArrayProperty::GetCPPType( FString* ExtendedTypeText/*=NULL*/, uint32 C
 	return TEXT("TArray");
 }
 
+FString UArrayProperty::GetCPPType( FString* ExtendedTypeText/*=NULL*/, uint32 CPPExportFlags/*=0*/ ) const
+{
+	checkSlow(Inner);
+	FString InnerExtendedTypeText;
+	FString InnerTypeText;
+	if ( ExtendedTypeText != NULL )
+	{
+		InnerTypeText = Inner->GetCPPType(&InnerExtendedTypeText, CPPExportFlags & ~CPPF_ArgumentOrReturnValue); // we won't consider array inners to be "arguments or return values"
+	}
+	return GetCPPTypeCustom(ExtendedTypeText, CPPExportFlags, InnerTypeText, InnerExtendedTypeText);
+}
+
 FString UArrayProperty::GetCPPTypeForwardDeclaration() const
 {
 	checkSlow(Inner);
diff --git a/Engine/Source/Runtime/CoreUObject/Private/UObject/PropertyAssetClass.cpp b/Engine/Source/Runtime/CoreUObject/Private/UObject/PropertyAssetClass.cpp
--- a/Engine/Source/Runtime/CoreUObject/Private/UObject/PropertyAssetClass.cpp
+++ b/Engine/Source/Runtime/CoreUObject/Private/UObject/PropertyAssetClass.cpp
@@ -6,9 +6,14 @@
 	UAssetClassProperty.
 -----------------------------------------------------------------------------*/
 
-FString UAssetClassProperty::GetCPPType( FString* ExtendedTypeText/*=NULL*/, uint32 CPPExportFlags/*=0*/ ) const
+FString UAssetClassProperty::GetCPPType(FString* ExtendedTypeText, uint32 CPPExportFlags) const
 {
-	return FString::Printf(TEXT("TAssetSubclassOf<%s%s> "),MetaClass->GetPrefixCPP(),*MetaClass->GetName());
+	return GetCPPTypeCustom(ExtendedTypeText, CPPExportFlags, MetaClass);
+}
+FString UAssetClassProperty::GetCPPTypeCustom(FString* ExtendedTypeText, uint32 CPPExportFlags, UClass* ActualClass) const
+{
+	check(ActualClass);
+	return FString::Printf(TEXT("TAssetSubclassOf<%s%s> "), ActualClass->GetPrefixCPP(), *ActualClass->GetName());
 }
 FString UAssetClassProperty::GetCPPMacroType( FString& ExtendedTypeText ) const
 {
diff --git a/Engine/Source/Runtime/CoreUObject/Private/UObject/PropertyAssetObject.cpp b/Engine/Source/Runtime/CoreUObject/Private/UObject/PropertyAssetObject.cpp
--- a/Engine/Source/Runtime/CoreUObject/Private/UObject/PropertyAssetObject.cpp
+++ b/Engine/Source/Runtime/CoreUObject/Private/UObject/PropertyAssetObject.cpp
@@ -6,9 +6,10 @@
 	UAssetObjectProperty.
 -----------------------------------------------------------------------------*/
 
-FString UAssetObjectProperty::GetCPPType( FString* ExtendedTypeText/*=NULL*/, uint32 CPPExportFlags/*=0*/ ) const
+FString UAssetObjectProperty::GetCPPTypeCustom(FString* ExtendedTypeText, uint32 CPPExportFlags, UClass* ActualClass) const
 {
-	return FString::Printf( TEXT("TAssetPtr<%s%s>"), PropertyClass->GetPrefixCPP(), *PropertyClass->GetName() );
+	check(ActualClass);
+	return FString::Printf(TEXT("TAssetPtr<%s%s>"), ActualClass->GetPrefixCPP(), *ActualClass->GetName());
 }
 FString UAssetObjectProperty::GetCPPMacroType( FString& ExtendedTypeText ) const
 {
diff --git a/Engine/Source/Runtime/CoreUObject/Private/UObject/PropertyClass.cpp b/Engine/Source/Runtime/CoreUObject/Private/UObject/PropertyClass.cpp
--- a/Engine/Source/Runtime/CoreUObject/Private/UObject/PropertyClass.cpp
+++ b/Engine/Source/Runtime/CoreUObject/Private/UObject/PropertyClass.cpp
@@ -89,11 +89,17 @@ const TCHAR* UClassProperty::ImportText_Internal( const TCHAR* Buffer, void* Dat
 	return Result;
 }
 
-FString UClassProperty::GetCPPType( FString* ExtendedTypeText/*=NULL*/, uint32 CPPExportFlags/*=0*/ ) const
+FString UClassProperty::GetCPPType(FString* ExtendedTypeText, uint32 CPPExportFlags) const
+{
+	return GetCPPTypeCustom(ExtendedTypeText, CPPExportFlags, MetaClass);
+}
+
+FString UClassProperty::GetCPPTypeCustom(FString* ExtendedTypeText, uint32 CPPExportFlags, UClass* ActualClass) const
 {
 	if (PropertyFlags & CPF_UObjectWrapper)
 	{
-		return FString::Printf(TEXT("TSubclassOf<%s%s> "),MetaClass->GetPrefixCPP(),*MetaClass->GetName());
+		check(ActualClass);
+		return FString::Printf(TEXT("TSubclassOf<%s%s> "), ActualClass->GetPrefixCPP(), *ActualClass->GetName());
 	}
 	else
 	{
diff --git a/Engine/Source/Runtime/CoreUObject/Private/UObject/PropertyObject.cpp b/Engine/Source/Runtime/CoreUObject/Private/UObject/PropertyObject.cpp
--- a/Engine/Source/Runtime/CoreUObject/Private/UObject/PropertyObject.cpp
+++ b/Engine/Source/Runtime/CoreUObject/Private/UObject/PropertyObject.cpp
@@ -8,9 +8,10 @@
 	UObjectProperty.
 -----------------------------------------------------------------------------*/
 
-FString UObjectProperty::GetCPPType( FString* ExtendedTypeText/*=NULL*/, uint32 CPPExportFlags/*=0*/ ) const
+FString UObjectProperty::GetCPPTypeCustom(FString* ExtendedTypeText, uint32 CPPExportFlags, UClass* ActualClass)  const
 {
-	return FString::Printf( TEXT("%s%s*"), PropertyClass->GetPrefixCPP(), *PropertyClass->GetName() );
+	check(ActualClass);
+	return FString::Printf(TEXT("%s%s*"), ActualClass->GetPrefixCPP(), *ActualClass->GetName());
 }
 
 FString UObjectProperty::GetCPPTypeForwardDeclaration() const
diff --git a/Engine/Source/Runtime/CoreUObject/Public/UObject/AssetPtr.h b/Engine/Source/Runtime/CoreUObject/Public/UObject/AssetPtr.h
--- a/Engine/Source/Runtime/CoreUObject/Public/UObject/AssetPtr.h
+++ b/Engine/Source/Runtime/CoreUObject/Public/UObject/AssetPtr.h
@@ -297,6 +297,9 @@ template <class T> struct TIsWeakPointerType<TAssetPtr<T> > { enum { Value = TIs
 template<class TClass>
 class TAssetSubclassOf
 {
+	template <class U>
+	friend class TAssetSubclassOf;
+
 public:
 	/** NULL constructor **/
 	FORCEINLINE TAssetSubclassOf()
diff --git a/Engine/Source/Runtime/CoreUObject/Public/UObject/UnrealType.h b/Engine/Source/Runtime/CoreUObject/Public/UObject/UnrealType.h
--- a/Engine/Source/Runtime/CoreUObject/Public/UObject/UnrealType.h
+++ b/Engine/Source/Runtime/CoreUObject/Public/UObject/UnrealType.h
@@ -105,7 +105,8 @@ public:
 											FOutputDevice* Warn, TArray<struct FDefinedProperty>& DefinedProperties );
 
 	// UHT interface
-	void ExportCppDeclaration(FOutputDevice& Out, EExportedDeclaration::Type DeclarationType, const TCHAR* ArrayDimOverride = NULL, uint32 AdditionalExportCPPFlags = 0, bool bSkipParameterName = false) const;
+	void ExportCppDeclaration(FOutputDevice& Out, EExportedDeclaration::Type DeclarationType, const TCHAR* ArrayDimOverride = NULL, uint32 AdditionalExportCPPFlags = 0
+		, bool bSkipParameterName = false, const FString* ActualCppType = nullptr, const FString* ActualExtendedType = nullptr) const;
 	virtual FString GetCPPMacroType( FString& ExtendedTypeText ) const;
 	virtual bool PassCPPArgsByRef() const { return false; }
 
@@ -1751,6 +1752,8 @@ class COREUOBJECT_API UObjectPropertyBase : public UProperty
 	// UObjectPropertyBase interface
 public:
 
+	virtual FString GetCPPTypeCustom(FString* ExtendedTypeText, uint32 CPPExportFlags, UClass* ActualClass) const PURE_VIRTUAL(UObjectPropertyBase::GetCPPTypeCustom, return TEXT(""););
+
 	/**
 	 * Parses a text buffer into an object reference.
 	 *
@@ -1812,7 +1815,6 @@ protected:
 
 	virtual void CheckValidObject(void* Value) const;
 	// End of UObjectPropertyBase interface
-
 };
 
 template<typename InTCppType>
@@ -1855,6 +1857,11 @@ public:
 		return TIsWeakPointerType<InTCppType>::Value;
 	}
 	// End of UProperty interface
+
+	virtual FString GetCPPType(FString* ExtendedTypeText, uint32 CPPExportFlags) const override
+	{
+		return GetCPPTypeCustom(ExtendedTypeText, CPPExportFlags, PropertyClass);
+	}
 };
 
 
@@ -1877,7 +1884,6 @@ class COREUOBJECT_API UObjectProperty : public TUObjectPropertyBase<UObject*>
 
 	// UHT interface
 	virtual FString GetCPPMacroType( FString& ExtendedTypeText ) const  override;
-	virtual FString GetCPPType( FString* ExtendedTypeText, uint32 CPPExportFlags ) const override;
 	virtual FString GetCPPTypeForwardDeclaration() const override;
 	// End of UHT interface
 
@@ -1899,6 +1905,8 @@ public:
 		return GetPropertyValue(PropertyValueAddress);
 	}
 	virtual void SetObjectPropertyValue(void* PropertyValueAddress, UObject* Value) const override;
+
+	virtual FString GetCPPTypeCustom(FString* ExtendedTypeText, uint32 CPPExportFlags, UClass* ActualClass)  const override;
 	// End of UObjectPropertyBase interface
 };
 
@@ -2008,7 +2016,6 @@ class COREUOBJECT_API UAssetObjectProperty : public TUObjectPropertyBase<FAssetP
 
 	// UHT interface
 	virtual FString GetCPPMacroType( FString& ExtendedTypeText ) const  override;
-	virtual FString GetCPPType( FString* ExtendedTypeText, uint32 CPPExportFlags ) const override;
 	virtual FString GetCPPTypeForwardDeclaration() const override;
 	// End of UHT interface
 
@@ -2033,6 +2040,8 @@ class COREUOBJECT_API UAssetObjectProperty : public TUObjectPropertyBase<FAssetP
 	{
 		return true;
 	}
+	virtual FString GetCPPTypeCustom(FString* ExtendedTypeText, uint32 CPPExportFlags, UClass* ActualClass)  const override;
+
 private:
 	virtual uint32 GetValueTypeHashInternal(const void* Src) const override
 	{
@@ -2097,8 +2106,8 @@ public:
 	// End of UObject interface
 
 	// UHT interface
+	virtual FString GetCPPType(FString* ExtendedTypeText, uint32 CPPExportFlags)  const override;
 	virtual FString GetCPPMacroType( FString& ExtendedTypeText ) const  override;
-	virtual FString GetCPPType( FString* ExtendedTypeText, uint32 CPPExportFlags ) const override;
 	virtual FString GetCPPTypeForwardDeclaration() const override;
 	// End of UHT interface
 
@@ -2107,6 +2116,8 @@ public:
 	virtual bool SameType(const UProperty* Other) const override;
 	// End of UProperty interface
 
+	virtual FString GetCPPTypeCustom(FString* ExtendedTypeText, uint32 CPPExportFlags, UClass* ActualClass)  const override;
+
 	/**
 	 * Setter function for this property's MetaClass member. Favor this function 
 	 * whilst loading (since, to handle circular dependencies, we defer some 
@@ -2151,8 +2162,8 @@ public:
 	{}
 
 	// UHT interface
+	virtual FString GetCPPType(FString* ExtendedTypeText, uint32 CPPExportFlags) const override;
 	virtual FString GetCPPMacroType( FString& ExtendedTypeText ) const  override;
-	virtual FString GetCPPType( FString* ExtendedTypeText, uint32 CPPExportFlags ) const override;
 	virtual FString GetCPPTypeForwardDeclaration() const override;
 	// End of UHT interface
 
@@ -2164,6 +2175,8 @@ public:
 	// UProperty interface
 	virtual bool SameType(const UProperty* Other) const override;
 	// End of UProperty interface
+
+	virtual FString GetCPPTypeCustom(FString* ExtendedTypeText, uint32 CPPExportFlags, UClass* ActualClass)  const override;
 };
 
 /*-----------------------------------------------------------------------------
@@ -2387,6 +2400,8 @@ public:
 	virtual void EmitReferenceInfo(UClass& OwnerClass, int32 BaseOffset) override;
 	virtual bool SameType(const UProperty* Other) const override;
 	// End of UProperty interface
+
+	FString GetCPPTypeCustom(FString* ExtendedTypeText, uint32 CPPExportFlags, const FString& InnerTypeText, const FString& InInnerExtendedTypeText) const;
 };
 
 // need to break this out a different type so that the DECLARE_CASTED_CLASS_INTRINSIC macro can digest the comma
diff --git a/Engine/Source/Runtime/Engine/Public/GeneratedCodeHelpers.h b/Engine/Source/Runtime/Engine/Public/GeneratedCodeHelpers.h
--- a/Engine/Source/Runtime/Engine/Public/GeneratedCodeHelpers.h
+++ b/Engine/Source/Runtime/Engine/Public/GeneratedCodeHelpers.h
@@ -17,6 +17,13 @@
 // Special libraries
 #include "Kismet/DataTableFunctionLibrary.h"
 
+template<class NativeType>
+inline NativeType* NoNativeCast(UClass* NoNativeClass, UObject* Object)
+{
+	check(NoNativeClass);
+	check(!Object || (nullptr != Cast<NativeType>(Object)));
+	return (Object && Object->IsA(NoNativeClass)) ? ((NativeType*)Object) : nullptr;
+}
 
 inline UClass* DynamicMetaCast(const UClass* DesiredClass, UClass* SourceClass)
 {
