BlueprintCompilerCppBackend:
- support for unconverted types - WIP
- use "UsedAssets" array from the class instead of loading objects

diff --git a/Engine/Source/Developer/BlueprintCompilerCppBackend/Private/BlueprintCompilerCppBackend.cpp b/Engine/Source/Developer/BlueprintCompilerCppBackend/Private/BlueprintCompilerCppBackend.cpp
--- a/Engine/Source/Developer/BlueprintCompilerCppBackend/Private/BlueprintCompilerCppBackend.cpp
+++ b/Engine/Source/Developer/BlueprintCompilerCppBackend/Private/BlueprintCompilerCppBackend.cpp
@@ -162,7 +162,7 @@ void FBlueprintCompilerCppBackend::EmitDynamicCastStatement(FEmitterLocalContext
 	if (BPGC && !EmitterContext.Dependencies.WillClassBeConverted(BPGC))
 	{
 		const FString NativeClass = FEmitHelper::GetCppName(EmitterContext.GetFirstNativeParent(ClassPtr));
-		const FString TargetClass = EmitterContext.FindGloballyMappedObject(ClassPtr, true);
+		const FString TargetClass = EmitterContext.FindGloballyMappedObject(ClassPtr, UClass::StaticClass(), true);
 		EmitterContext.AddLine(FString::Printf(TEXT("%s = NoNativeCast<%s>(%s, %s);"), *CastedValue, *NativeClass, *TargetClass, *ObjectValue));
 	}
 	else
diff --git a/Engine/Source/Developer/BlueprintCompilerCppBackend/Private/BlueprintCompilerCppBackendGatherDependencies.cpp b/Engine/Source/Developer/BlueprintCompilerCppBackend/Private/BlueprintCompilerCppBackendGatherDependencies.cpp
--- a/Engine/Source/Developer/BlueprintCompilerCppBackend/Private/BlueprintCompilerCppBackendGatherDependencies.cpp
+++ b/Engine/Source/Developer/BlueprintCompilerCppBackend/Private/BlueprintCompilerCppBackendGatherDependencies.cpp
@@ -123,7 +123,7 @@ void FGatherConvertedClassDependencies::HandleObjectReference(UObject*& InObject
 	}
 	else if ((InObject->IsAsset() || ObjAsBPGC) && !InObject->IsIn(OriginalStruct))
 	{
-		Assets.Add(InObject);
+		Assets.AddUnique(InObject);
 		return;
 	}
 	else if (auto ObjAsClass = Cast<UClass>(InObject))
diff --git a/Engine/Source/Developer/BlueprintCompilerCppBackend/Private/BlueprintCompilerCppBackendUMG.cpp b/Engine/Source/Developer/BlueprintCompilerCppBackend/Private/BlueprintCompilerCppBackendUMG.cpp
--- a/Engine/Source/Developer/BlueprintCompilerCppBackend/Private/BlueprintCompilerCppBackendUMG.cpp
+++ b/Engine/Source/Developer/BlueprintCompilerCppBackend/Private/BlueprintCompilerCppBackendUMG.cpp
@@ -77,7 +77,7 @@ void FBackendHelperUMG::EmitWidgetInitializationFunctions(FEmitterLocalContext&
 			Context.AddLine(TEXT("{"));
 			Context.IncreaseIndent();
 
-			const FString WidgetTreeStr = Context.FindGloballyMappedObject(WidgetClass->WidgetTree, true);
+			const FString WidgetTreeStr = Context.FindGloballyMappedObject(WidgetClass->WidgetTree, UWidgetTree::StaticClass(), true);
 			ensure(!WidgetTreeStr.IsEmpty());
 			const FString AnimationsArrayNativeName = GenerateLocalProperty(Context, FindFieldChecked<UArrayProperty>(UWidgetBlueprintGeneratedClass::StaticClass(), TEXT("Animations")), reinterpret_cast<const uint8*>(&WidgetClass->Animations));
 			const FString BindingsArrayNativeName = GenerateLocalProperty(Context, FindFieldChecked<UArrayProperty>(UWidgetBlueprintGeneratedClass::StaticClass(), TEXT("Bindings")), reinterpret_cast<const uint8*>(&WidgetClass->Bindings));
diff --git a/Engine/Source/Developer/BlueprintCompilerCppBackend/Private/BlueprintCompilerCppBackendUtils.cpp b/Engine/Source/Developer/BlueprintCompilerCppBackend/Private/BlueprintCompilerCppBackendUtils.cpp
--- a/Engine/Source/Developer/BlueprintCompilerCppBackend/Private/BlueprintCompilerCppBackendUtils.cpp
+++ b/Engine/Source/Developer/BlueprintCompilerCppBackend/Private/BlueprintCompilerCppBackendUtils.cpp
@@ -23,7 +23,7 @@ FString FEmitterLocalContext::GenerateUniqueLocalName()
 	return UniqueName;
 }
 
-FString FEmitterLocalContext::FindGloballyMappedObject(UObject* Object, bool bLoadIfNotFound)
+FString FEmitterLocalContext::FindGloballyMappedObject(const UObject* Object, const UClass* ExpectedClass, bool bLoadIfNotFound, bool bTryUsedAssetsList)
 {
 	UClass* ActualClass = Cast<UClass>(Dependencies.GetOriginalStruct());
 	if (ActualClass && Object && Object->IsIn(ActualClass))
@@ -38,34 +38,35 @@ FString FEmitterLocalContext::FindGloballyMappedObject(UObject* Object, bool bLo
 			return *NamePtr;
 		}
 
-		int32 ObjectsCreatedPerClassIdx = INDEX_NONE;
-		if (MiscConvertedSubobjects.Find(Object, ObjectsCreatedPerClassIdx))
+		const UClass* ObjectClassToUse = ExpectedClass ? ExpectedClass : GetNativeOrConvertedClass(Object->GetClass());
+		const FString ClassString = FEmitHelper::GetCppName(ObjectClassToUse);
+
+		int32 ObjectsCreatedPerClassIdx = MiscConvertedSubobjects.IndexOfByKey(Object);
+		if (INDEX_NONE != ObjectsCreatedPerClassIdx)
 		{
 			return FString::Printf(TEXT("CastChecked<%s>(CastChecked<UDynamicClass>(%s::StaticClass())->MiscConvertedSubobjects[%d])")
-				, *FEmitHelper::GetCppName(Object->GetClass())
-				, *FEmitHelper::GetCppName(ActualClass)
-				, ObjectsCreatedPerClassIdx);
+				, *ClassString, *FEmitHelper::GetCppName(ActualClass), ObjectsCreatedPerClassIdx);
 		}
-		if (DynamicBindingObjects.Find(Object, ObjectsCreatedPerClassIdx))
+
+		ObjectsCreatedPerClassIdx = ObjectsCreatedPerClassIdx = DynamicBindingObjects.IndexOfByKey(Object);
+		if (INDEX_NONE != ObjectsCreatedPerClassIdx)
 		{
 			return FString::Printf(TEXT("CastChecked<%s>(CastChecked<UDynamicClass>(%s::StaticClass())->DynamicBindingObjects[%d])")
-				, *FEmitHelper::GetCppName(Object->GetClass())
-				, *FEmitHelper::GetCppName(ActualClass)
-				, ObjectsCreatedPerClassIdx);
+				, *ClassString, *FEmitHelper::GetCppName(ActualClass), ObjectsCreatedPerClassIdx);
 		}
-		if (ComponentTemplates.Find(Object, ObjectsCreatedPerClassIdx))
+
+		ObjectsCreatedPerClassIdx = ComponentTemplates.IndexOfByKey(Object);
+		if (INDEX_NONE != ObjectsCreatedPerClassIdx)
 		{
 			return FString::Printf(TEXT("CastChecked<%s>(CastChecked<UDynamicClass>(%s::StaticClass())->ComponentTemplates[%d])")
-				, *FEmitHelper::GetCppName(Object->GetClass())
-				, *FEmitHelper::GetCppName(ActualClass)
-				, ObjectsCreatedPerClassIdx);
+				, *ClassString, *FEmitHelper::GetCppName(ActualClass), ObjectsCreatedPerClassIdx);
 		}
-		if (Timelines.Find(Object, ObjectsCreatedPerClassIdx))
+
+		ObjectsCreatedPerClassIdx = Timelines.IndexOfByKey(Object);
+		if (INDEX_NONE != ObjectsCreatedPerClassIdx)
 		{
 			return FString::Printf(TEXT("CastChecked<%s>(CastChecked<UDynamicClass>(%s::StaticClass())->Timelines[%d])")
-				, *FEmitHelper::GetCppName(Object->GetClass())
-				, *FEmitHelper::GetCppName(ActualClass)
-				, ObjectsCreatedPerClassIdx);
+				, *ClassString, *FEmitHelper::GetCppName(ActualClass), ObjectsCreatedPerClassIdx);
 		}
 
 		if ((CurrentCodeType == EGeneratedCodeType::SubobjectsOfClass) || (CurrentCodeType == EGeneratedCodeType::CommonConstructor))
@@ -109,12 +110,30 @@ FString FEmitterLocalContext::FindGloballyMappedObject(UObject* Object, bool bLo
 	}
 
 	// TODO: handle subobjects
-
-	if (bLoadIfNotFound && ensure(Object))
+	ensure(!bLoadIfNotFound || Object);
+	if (Object && (bLoadIfNotFound || bTryUsedAssetsList))
 	{
-		UClass* FoundClass = Object->GetClass();
-		const FString ClassString = FEmitHelper::GetCppName(FoundClass);
-		return FString::Printf(TEXT("LoadObject<%s>(nullptr, TEXT(\"%s\"))"), *ClassString, *(Object->GetPathName().ReplaceCharWithEscapedChar()));
+		const UClass* ObjectClassToUse = ExpectedClass ? ExpectedClass : GetNativeOrConvertedClass(Object->GetClass());
+		const FString ClassString = FEmitHelper::GetCppName(ObjectClassToUse);
+
+		if (bTryUsedAssetsList)
+		{
+			const int32 AssetIndex = Dependencies.Assets.IndexOfByKey(Object);
+			if (INDEX_NONE != AssetIndex)
+			{
+				return FString::Printf(TEXT("CastChecked<%s>(CastChecked<UDynamicClass>(%s::StaticClass())->UsedAssets[%d])")
+					, *ClassString
+					, *FEmitHelper::GetCppName(ActualClass)
+					, AssetIndex);
+			}
+		}
+
+		if (bLoadIfNotFound)
+		{
+			return FString::Printf(TEXT("LoadObject<%s>(nullptr, TEXT(\"%s\"))")
+				, *ClassString
+				, *(Object->GetPathName().ReplaceCharWithEscapedChar()));
+		}
 	}
 
 	return FString{};
@@ -138,13 +157,7 @@ FString FEmitterLocalContext::ExportTextItem(const UProperty* Property, const vo
 	{
 		if (UObject* Object = ObjectPropertyBase->GetObjectPropertyValue(PropertyValue))
 		{
-			UClass* ActualClass = ObjectPropertyBase->PropertyClass;
-			auto BPGC = Cast<UBlueprintGeneratedClass>(ActualClass);
-			const bool bConverted = BPGC && Dependencies.WillClassBeConverted(BPGC);
-			if (BPGC && !bConverted)
-			{
-				ActualClass = GetFirstNativeParent(BPGC);
-			}
+			const UClass* ActualClass = GetNativeOrConvertedClass(ObjectPropertyBase->PropertyClass);
 			return FString::Printf(TEXT("LoadObject<%s>(nullptr, TEXT(\"%s\"))")
 				, *FEmitHelper::GetCppName(ActualClass)
 				, *(Object->GetPathName().ReplaceCharWithEscapedChar()));
@@ -813,14 +826,14 @@ FString FEmitHelper::LiteralTerm(FEmitterLocalContext& EmitterContext, const FEd
 	{
 		if (LiteralObject)
 		{
-			const FString MappedObject = EmitterContext.FindGloballyMappedObject(LiteralObject);
+			UClass* FoundClass = Cast<UClass>(Type.PinSubCategoryObject.Get());
+			UClass* ObjectClassToUse = FoundClass ? EmitterContext.GetNativeOrConvertedClass(FoundClass) : UObject::StaticClass();
+			const FString MappedObject = EmitterContext.FindGloballyMappedObject(LiteralObject, ObjectClassToUse, true);
 			if (!MappedObject.IsEmpty())
 			{
 				return MappedObject;
 			}
-			UClass* FoundClass = Cast<UClass>(Type.PinSubCategoryObject.Get());
-			FString ClassString = FoundClass ? FEmitHelper::GetCppName(FoundClass) : TEXT("UObject");
-			return FString::Printf(TEXT("LoadObject<%s>(nullptr, TEXT(\"%s\"))"), *ClassString, *(LiteralObject->GetPathName().ReplaceCharWithEscapedChar()));
+			return FString(TEXT("nullptr"));
 		}
 		else
 		{
diff --git a/Engine/Source/Developer/BlueprintCompilerCppBackend/Private/BlueprintCompilerCppBackendUtils.h b/Engine/Source/Developer/BlueprintCompilerCppBackend/Private/BlueprintCompilerCppBackendUtils.h
--- a/Engine/Source/Developer/BlueprintCompilerCppBackend/Private/BlueprintCompilerCppBackendUtils.h
+++ b/Engine/Source/Developer/BlueprintCompilerCppBackend/Private/BlueprintCompilerCppBackendUtils.h
@@ -118,6 +118,21 @@ public:
 		return nullptr;
 	}
 
+	UClass* GetNativeOrConvertedClass(UClass* InClass) const
+	{
+		check(InClass);
+		if (InClass->HasAnyClassFlags(CLASS_Native))
+		{
+			return InClass;
+		}
+		auto BPGC = CastChecked<UBlueprintGeneratedClass>(InClass);
+		if (Dependencies.WillClassBeConverted(BPGC))
+		{
+			return InClass;
+		}
+		return GetFirstNativeParent(InClass);
+	}
+
 	FString GenerateGetProperty(const UProperty* Property) const;
 
 	FString GenerateUniqueLocalName();
@@ -129,7 +144,7 @@ public:
 
 	/** All objects (that can be referenced from other package) that will have a different path in cooked build
 	(due to the native code generation), should be handled by this function */
-	FString FindGloballyMappedObject(UObject* Object, bool bLoadIfNotFound = false);
+	FString FindGloballyMappedObject(const UObject* Object, const UClass* ExpectedClass = nullptr, bool bLoadIfNotFound = false, bool bTryUsedAssetsList = true);
 
 	// TEXT FUNCTIONS
 
diff --git a/Engine/Source/Developer/BlueprintCompilerCppBackend/Private/BlueprintCompilerCppBackendValueHelper.cpp b/Engine/Source/Developer/BlueprintCompilerCppBackend/Private/BlueprintCompilerCppBackendValueHelper.cpp
--- a/Engine/Source/Developer/BlueprintCompilerCppBackend/Private/BlueprintCompilerCppBackendValueHelper.cpp
+++ b/Engine/Source/Developer/BlueprintCompilerCppBackend/Private/BlueprintCompilerCppBackendValueHelper.cpp
@@ -182,7 +182,8 @@ FString FEmitDefaultValueHelper::HandleSpecialTypes(FEmitterLocalContext& Contex
 		if (Object)
 		{
 			{
-				const FString MappedObject = Context.FindGloballyMappedObject(Object);
+				UClass* ObjectClassToUse = Context.GetNativeOrConvertedClass(ObjectProperty->PropertyClass);
+				const FString MappedObject = Context.FindGloballyMappedObject(Object, ObjectClassToUse);
 				if (!MappedObject.IsEmpty())
 				{
 					return MappedObject;
@@ -299,7 +300,7 @@ FString FEmitDefaultValueHelper::HandleNonNativeComponent(FEmitterLocalContext&
 				}
 				else if (USceneComponent* ParentComponentTemplate = Node->GetParentComponentTemplate(CastChecked<UBlueprint>(BPGC->ClassGeneratedBy)))
 				{
-					ParentVariableName = Context.FindGloballyMappedObject(ParentComponentTemplate);
+					ParentVariableName = Context.FindGloballyMappedObject(ParentComponentTemplate, USceneComponent::StaticClass());
 				}
 
 				if (!ParentVariableName.IsEmpty())
@@ -357,7 +358,7 @@ struct FDependenciesHelper
 		}
 		for (auto LocAsset : Context.Dependencies.Assets)
 		{
-			const FString AssetStr = Context.FindGloballyMappedObject(LocAsset, true);
+			const FString AssetStr = Context.FindGloballyMappedObject(LocAsset, nullptr, true, false);
 			Context.AddLine(FString::Printf(TEXT("CastChecked<UDynamicClass>(GetClass())->UsedAssets.Add(%s);"), *AssetStr));
 		}
 	}
diff --git a/Engine/Source/Developer/BlueprintCompilerCppBackend/Public/BlueprintCompilerCppBackendGatherDependencies.h b/Engine/Source/Developer/BlueprintCompilerCppBackend/Public/BlueprintCompilerCppBackendGatherDependencies.h
--- a/Engine/Source/Developer/BlueprintCompilerCppBackend/Public/BlueprintCompilerCppBackendGatherDependencies.h
+++ b/Engine/Source/Developer/BlueprintCompilerCppBackend/Public/BlueprintCompilerCppBackendGatherDependencies.h
@@ -15,7 +15,9 @@ protected:
 
 public:
 	// Dependencies:
-	TSet<UObject*> Assets;
+
+	TArray<UObject*> Assets; 
+
 	TSet<UBlueprintGeneratedClass*> ConvertedClasses;
 	TSet<UUserDefinedStruct*> ConvertedStructs;
 	TSet<UUserDefinedEnum*> ConvertedEnum;
