Changes required to support unicode (or otherwise invalid) class names for blueprints converted to native types. We now add a postfix (__pf) to native types generated for blueprints. We also overwrite invalid symbols (ie, +, - or unicode characters) found in blueprint asset names. All dynamic classes now need to explicitly specify a name.

diff --git a/Engine/Source/Developer/BlueprintCompilerCppBackend/Private/BlueprintCompilerCppBackendUtils.cpp b/Engine/Source/Developer/BlueprintCompilerCppBackend/Private/BlueprintCompilerCppBackendUtils.cpp
--- a/Engine/Source/Developer/BlueprintCompilerCppBackend/Private/BlueprintCompilerCppBackendUtils.cpp
+++ b/Engine/Source/Developer/BlueprintCompilerCppBackend/Private/BlueprintCompilerCppBackendUtils.cpp
@@ -182,16 +182,24 @@ FString FEmitterLocalContext::ExportCppDeclaration(const UProperty* Property, EE
 	{
 		auto BPGC = Cast<UBlueprintGeneratedClass>(InActualClass);
 		const bool bConverted = BPGC && Dependencies.WillClassBeConverted(BPGC);
-		if (BPGC && !bConverted)
+		if (BPGC )
 		{
-			const bool bIsParameter = (DeclarationType == EExportedDeclaration::Parameter) || (DeclarationType == EExportedDeclaration::MacroParameter);
-			const uint32 LocalExportCPPFlags = ExportCPPFlags | (bIsParameter ? CPPF_ArgumentOrReturnValue : 0);
-			UClass* NativeType = GetFirstNativeParent(BPGC);
-			ActualCppType = ObjectPropertyBase->GetCPPTypeCustom(&ActualExtendedType, LocalExportCPPFlags, NativeType);
-			ActualCppTypePtr = &ActualCppType;
-			if (!ActualExtendedType.IsEmpty())
+			if (!bConverted)
 			{
-				ActualExtendedTypePtr = &ActualExtendedType;
+				const bool bIsParameter = (DeclarationType == EExportedDeclaration::Parameter) || (DeclarationType == EExportedDeclaration::MacroParameter);
+				const uint32 LocalExportCPPFlags = ExportCPPFlags | (bIsParameter ? CPPF_ArgumentOrReturnValue : 0);
+				UClass* NativeType = GetFirstNativeParent(BPGC);
+				ActualCppType = ObjectPropertyBase->GetCPPTypeCustom(&ActualExtendedType, LocalExportCPPFlags, NativeType);
+				ActualCppTypePtr = &ActualCppType;
+				if (!ActualExtendedType.IsEmpty())
+				{
+					ActualExtendedTypePtr = &ActualExtendedType;
+				}
+			}
+			else
+			{
+				ActualCppType = FEmitHelper::GetCppName(ObjectPropertyBase->PropertyClass, false) + TEXT("*");
+				ActualCppTypePtr = &ActualCppType;
 			}
 		}
 	};
@@ -215,6 +223,11 @@ FString FEmitterLocalContext::ExportCppDeclaration(const UProperty* Property, EE
 	{
 		GetActualNameCPP(ObjectProperty, ObjectProperty->PropertyClass);
 	}
+	else if (auto StructProperty = Cast<const UStructProperty>(Property))
+	{
+		ActualCppType = FEmitHelper::GetCppName(StructProperty->Struct, false);
+		ActualCppTypePtr = &ActualCppType;
+	}
 
 	if (ArrayProperty)
 	{
@@ -259,7 +272,14 @@ FString FEmitHelper::GetCppName(const UField* Field, bool bUInterface)
 				, *AsClass->GetName());
 		}
 		auto AsStruct = CastChecked<UStruct>(Field);
-		return FString::Printf(TEXT("%s%s"), AsStruct->GetPrefixCPP(), *AsStruct->GetName());
+		if (AsStruct->HasAnyFlags(RF_Native))
+		{
+			return FString::Printf(TEXT("%s%s"), AsStruct->GetPrefixCPP(), *AsStruct->GetName());
+		}
+		else
+		{
+			return ::UnicodeToCPPIdentifier(*AsStruct->GetName(), false, AsStruct->GetPrefixCPP());
+		}
 	}
 	else if (auto AsProperty = Cast<UProperty>(Field))
 	{
@@ -531,6 +551,11 @@ FString FEmitHelper::GenerateReplaceConvertedMD(UObject* Obj)
 		// TODO: It would be better to load all redirectors before compiling. Than checking here AssetRegistry.
 
 		Result += TEXT("\"");
+
+		// 4. Add overridden name:
+		Result += TEXT(", OverrideNativeName=\"");
+		Result += Obj->GetName();
+		Result += TEXT("\"");
 	}
 
 	return Result;
@@ -538,9 +563,28 @@ FString FEmitHelper::GenerateReplaceConvertedMD(UObject* Obj)
 
 FString FEmitHelper::GetBaseFilename(const UObject* AssetObj)
 {
-	// @TODO: See DanO about properly handling special asset names (w/ unicode, conflicts, etc.)
 	FString PackagePath = AssetObj->GetOutermost()->GetPathName();
-	return FPackageName::GetLongPackageAssetName(PackagePath);
+	// We have to sanitize the package path because UHT is going to generate header guards (preprocessor symbols)
+	// based on the filename. I'm also not interested in exploring the depth of unicode filename support in UHT,
+	// UBT, and our various c++ toolchains, so this logic is pretty aggressive:
+	FString Postfix = TEXT("__pf");
+	for (auto& Char : PackagePath)
+	{
+		if (Char == '/' || Char == '\\')
+		{
+			continue;
+		}
+
+		if (!IsValidCPPIdentifierChar(Char))
+		{
+			// deterministically map char to a valid ascii character, we have 63 characters available (aA-zZ, 0-9, and _)
+			// so the optimal encoding would be base 63:
+			Postfix.Append(ToValidCPPIdentifierChars(Char));
+			Char = TCHAR('x');
+		}
+	}
+
+	return FPackageName::GetLongPackageAssetName(PackagePath+Postfix);
 }
 
 FString FEmitHelper::EmitUFuntion(UFunction* Function, TArray<FString>* AdditinalMetaData)
diff --git a/Engine/Source/Programs/UnrealHeaderTool/Private/CodeGenerator.cpp b/Engine/Source/Programs/UnrealHeaderTool/Private/CodeGenerator.cpp
--- a/Engine/Source/Programs/UnrealHeaderTool/Private/CodeGenerator.cpp
+++ b/Engine/Source/Programs/UnrealHeaderTool/Private/CodeGenerator.cpp
@@ -888,6 +888,16 @@ FString FNativeClassHeaderGenerator::GetOverriddenName(const UField* Item)
 	return Item->GetName();
 }
 
+FString FNativeClassHeaderGenerator::GetOverriddenNameForLiteral(const UField* Item)
+{
+	FString OverriddenName = Item->GetMetaData(TEXT("OverrideNativeName"));
+	if (!OverriddenName.IsEmpty())
+	{
+		return TEXT("TEXT(\"") + OverriddenName + TEXT("\")");
+	}
+	return TEXT("\"") + Item->GetName() + TEXT("\"");
+}
+
 FString FNativeClassHeaderGenerator::PropertyNew(FString& Meta, UProperty* Prop, const FString& OuterString, const FString& PropMacro, const TCHAR* NameSuffix, const TCHAR* Spaces, const TCHAR* SourceStruct)
 {
 	FString ExtraArgs;
@@ -1371,7 +1381,7 @@ void FNativeClassHeaderGenerator::ExportNativeGeneratedInitCode(FClass* Class, F
 			// Emit code to construct each UFunction and rebuild the function map at runtime
 			for (UFunction* Function : FunctionsInMap)
 			{
-				GeneratedClassRegisterFunctionText.Logf(TEXT("\t\t\t\tOuterClass->AddFunctionToFunctionMapWithOverriddenName(%s, \"%s\");%s\r\n"), *GetSingletonName(Function), *FNativeClassHeaderGenerator::GetOverriddenName(Function), *GetGeneratedCodeCRCTag(Function));
+				GeneratedClassRegisterFunctionText.Logf(TEXT("\t\t\t\tOuterClass->AddFunctionToFunctionMapWithOverriddenName(%s, %s);%s\r\n"), *GetSingletonName(Function), *FNativeClassHeaderGenerator::GetOverriddenNameForLiteral(Function), *GetGeneratedCodeCRCTag(Function));
 			}
 		}
 
@@ -1432,10 +1442,11 @@ void FNativeClassHeaderGenerator::ExportNativeGeneratedInitCode(FClass* Class, F
 	SingletonName.ReplaceInline(TEXT("()"), TEXT(""), ESearchCase::CaseSensitive); // function address
 
 	{	
+		FString OverridenClassName = *FNativeClassHeaderGenerator::GetOverriddenName(Class);
 		const TCHAR* ClassNameCPP = NameLookupCPP.GetNameCPP(Class);
 
 		GeneratedFunctionText.Logf(TEXT("\tstatic FCompiledInDefer Z_CompiledInDefer_UClass_%s(%s, &%s::StaticClass, TEXT(\"%s\"), %s);\r\n"),
-			ClassNameCPP, *SingletonName, ClassNameCPP, ClassNameCPP,
+			ClassNameCPP, *SingletonName, ClassNameCPP, *OverridenClassName,
 			bIsDynamic ? TEXT("true") : TEXT("false"));
 		
 		// Append base class' CRC at the end of the generated code, this will force update derived classes
@@ -1457,7 +1468,7 @@ void FNativeClassHeaderGenerator::ExportNativeGeneratedInitCode(FClass* Class, F
 		}
 		else
 		{
-			GeneratedPackageCPP.Logf(TEXT("\tIMPLEMENT_DYNAMIC_CLASS(%s, %u);\r\n"), ClassNameCPP, ClassCrc);			
+			GeneratedPackageCPP.Logf(TEXT("\tIMPLEMENT_DYNAMIC_CLASS(%s, TEXT(\"%s\"), %u);\r\n"), ClassNameCPP, *OverridenClassName, ClassCrc);
 		}
 	}
 }
@@ -1606,9 +1617,9 @@ void FNativeClassHeaderGenerator::ExportNatives(FClass* Class)
 
 		for (auto Func : FunctionsToExport)
 		{
-			GeneratedPackageCPP.Logf(TEXT("\t\tFNativeFunctionRegistrar::RegisterFunction(%s::StaticClass(),\"%s\",(Native)&%s::exec%s);\r\n"),
+			GeneratedPackageCPP.Logf(TEXT("\t\tFNativeFunctionRegistrar::RegisterFunction(%s::StaticClass(), %s,(Native)&%s::exec%s);\r\n"),
 				NameLookupCPP.GetNameCPP(Class),
-				*FNativeClassHeaderGenerator::GetOverriddenName(Func),
+				*FNativeClassHeaderGenerator::GetOverriddenNameForLiteral(Func),
 				Class->HasAnyClassFlags(CLASS_Interface) ? *FString::Printf(TEXT("I%s"), *Class->GetName()) : NameLookupCPP.GetNameCPP(Class),
 				*Func->GetName()
 			);
diff --git a/Engine/Source/Programs/UnrealHeaderTool/Private/NativeClassExporter.h b/Engine/Source/Programs/UnrealHeaderTool/Private/NativeClassExporter.h
--- a/Engine/Source/Programs/UnrealHeaderTool/Private/NativeClassExporter.h
+++ b/Engine/Source/Programs/UnrealHeaderTool/Private/NativeClassExporter.h
@@ -192,6 +192,11 @@ private:
 	 * This function returns the original name of the field (rather than the native, symbol name).
 	 */
 	FString GetOverriddenName(const UField* Item);
+	
+	/** 
+	 * Returns the name (overridden if marked up) with TEXT("") or "" wrappers for use in a string literal.
+	 */
+	FString GetOverriddenNameForLiteral(const UField* Item);
 
 	/** 
 	 * Export functions used to find and call C++ or script implementation of a script function in the interface 
diff --git a/Engine/Source/Runtime/CoreUObject/Private/UObject/Class.cpp b/Engine/Source/Runtime/CoreUObject/Private/UObject/Class.cpp
--- a/Engine/Source/Runtime/CoreUObject/Private/UObject/Class.cpp
+++ b/Engine/Source/Runtime/CoreUObject/Private/UObject/Class.cpp
@@ -69,6 +69,10 @@ void FNativeFunctionRegistrar::RegisterFunction(class UClass* Class, const ANSIC
 	Class->AddNativeFunction(InName, InPointer);
 }
 
+void FNativeFunctionRegistrar::RegisterFunction(class UClass* Class, const WIDECHAR* InName, Native InPointer)
+{
+	Class->AddNativeFunction(InName, InPointer);
+}
 
 /*-----------------------------------------------------------------------------
 	UField implementation.
@@ -3972,6 +3976,27 @@ void UClass::AddNativeFunction(const ANSICHAR* InName,Native InPointer)
 	new(NativeFunctionLookupTable) FNativeFunctionLookup(InFName,InPointer);
 }
 
+void UClass::AddNativeFunction(const WIDECHAR* InName, Native InPointer)
+{
+	FName InFName(InName);
+#if WITH_HOT_RELOAD
+	if (GIsHotReload)
+	{
+		// Find the function in the class's native function lookup table.
+		if (ReplaceNativeFunction(InFName, InPointer, true))
+		{
+			return;
+		}
+		else
+		{
+			// function was not found, so it's new
+			UE_LOG(LogClass, Log, TEXT("Function %s is new."), *InFName.ToString());
+		}
+	}
+#endif
+	new(NativeFunctionLookupTable)FNativeFunctionLookup(InFName, InPointer);
+}
+
 UFunction* UClass::FindFunctionByName(FName InName, EIncludeSuperFlag::Type IncludeSuper) const
 {
 	UFunction* Result = FuncMap.FindRef(InName);
diff --git a/Engine/Source/Runtime/CoreUObject/Private/UObject/ScriptCore.cpp b/Engine/Source/Runtime/CoreUObject/Private/UObject/ScriptCore.cpp
--- a/Engine/Source/Runtime/CoreUObject/Private/UObject/ScriptCore.cpp
+++ b/Engine/Source/Runtime/CoreUObject/Private/UObject/ScriptCore.cpp
@@ -142,41 +142,55 @@ FEditorScriptExecutionGuard::~FEditorScriptExecutionGuard()
 	GAllowActorScriptExecutionInEditor = bOldGAllowScriptExecutionInEditor;
 }
 
-FString UnicodeToCPPIdentifier(const FString& InName, bool bDeprecated, const TCHAR* Prefix)
+bool IsValidCPPIdentifierChar(TCHAR Char)
 {
-	// FName's can contain unicode characters or collide with other CPP identifiers or keywords. This function 
-	// returns a string that will have a prefix which is unlikely to collide with existing identifiers and
-	// converts unicode characters in place to valid ascii characters. Strictly speaking a C++ compiler *could*
-	// support unicode identifiers in source files, but I am not comfortable relying on this behavior.
-	
-	auto IsValidIdentifierChar = [](TCHAR Char)
-	{
-		return Char == TCHAR('_')
-			|| (Char >= TCHAR('a') && Char <= TCHAR('z'))
-			|| (Char >= TCHAR('A') && Char <= TCHAR('Z'))
-			|| (Char >= TCHAR('0') && Char <= TCHAR('9'));
-	};
+	return Char == TCHAR('_')
+		|| (Char >= TCHAR('a') && Char <= TCHAR('z'))
+		|| (Char >= TCHAR('A') && Char <= TCHAR('Z'))
+		|| (Char >= TCHAR('0') && Char <= TCHAR('9'));
+}
 
-	auto ToValidIdentifier = [](int32 Val) -> TCHAR
+FString ToValidCPPIdentifierChars(TCHAR Char)
+{
+	FString Ret;
+	int32 RawValue = Char;
+	int32 Counter = 0;
+	while (RawValue != 0)
 	{
-		if (Val < 26)
+		int32 Digit = RawValue % 63;
+		RawValue = (RawValue - Digit) / 63;
+
+		TCHAR SafeChar;
+		if (Digit < 26)
 		{
-			return TCHAR(TCHAR('a') + (26 - Val));
+			SafeChar = TCHAR(TCHAR('a') + (26 - Digit));
 		}
-		else if (Val < 52)
+		else if (Digit < 52)
 		{
-			return TCHAR(TCHAR('A') + (52 - Val));
+			SafeChar = TCHAR(TCHAR('A') + (52 - Digit));
 		}
-		else if (Val < 62)
+		else if (Digit < 62)
 		{
-			return TCHAR(TCHAR('0') + (62 - Val));
+			SafeChar = TCHAR(TCHAR('0') + (62 - Digit));
 		}
 		else
 		{
-			check(Val == 62);
-			return TCHAR('_');
+			check(Digit == 62);
+			SafeChar = TCHAR('_');
 		}
-	};
+
+		Ret.AppendChar(SafeChar);
+	}
+	return Ret;
+}
+
+FString UnicodeToCPPIdentifier(const FString& InName, bool bDeprecated, const TCHAR* Prefix)
+{
+	// FName's can contain unicode characters or collide with other CPP identifiers or keywords. This function 
+	// returns a string that will have a prefix which is unlikely to collide with existing identifiers and
+	// converts unicode characters in place to valid ascii characters. Strictly speaking a C++ compiler *could*
+	// support unicode identifiers in source files, but I am not comfortable relying on this behavior.
+
 
 	FString Ret = InName;
 	// Initialize postfix with a unique identifier. This prevents potential collisions between names that have unicode
@@ -185,18 +199,11 @@ FString UnicodeToCPPIdentifier(const FString& InName, bool bDeprecated, const TC
 	for (auto& Char : Ret)
 	{
 		// if the character is not a valid character for a c++ identifier, then we need to encode it using valid characters:
-		if (!IsValidIdentifierChar(Char))
+		if (!IsValidCPPIdentifierChar(Char))
 		{
 			// deterministically map char to a valid ascii character, we have 63 characters available (aA-zZ, 0-9, and _)
 			// so the optimal encoding would be base 63:
-			int32 RawValue = Char;
-			int32 Counter = 0;
-			while (RawValue != 0)
-			{
-				int32 Digit = RawValue % 63;
-				RawValue = (RawValue - Digit) / 63;
-				Postfix.AppendChar(ToValidIdentifier(Digit));
-			}
+			Postfix.Append(ToValidCPPIdentifierChars(Char));
 			Char = TCHAR('x');
 		}
 	}
diff --git a/Engine/Source/Runtime/CoreUObject/Private/UObject/UObjectBase.cpp b/Engine/Source/Runtime/CoreUObject/Private/UObject/UObjectBase.cpp
--- a/Engine/Source/Runtime/CoreUObject/Private/UObject/UObjectBase.cpp
+++ b/Engine/Source/Runtime/CoreUObject/Private/UObject/UObjectBase.cpp
@@ -690,7 +690,7 @@ void UObjectCompiledInDefer(UClass *(*InRegister)(), UClass *(*InStaticClass)(),
 		FClassConstructFunctions ClassFunctions;
 		ClassFunctions.ZConstructFn = InRegister;
 		ClassFunctions.StaticClassFn = InStaticClass;
-		GetDynamicClassMap().Add(FName(Name + 1), ClassFunctions);
+		GetDynamicClassMap().Add(FName(Name), ClassFunctions);
 	}
 }
 
diff --git a/Engine/Source/Runtime/CoreUObject/Public/UObject/Class.h b/Engine/Source/Runtime/CoreUObject/Public/UObject/Class.h
--- a/Engine/Source/Runtime/CoreUObject/Public/UObject/Class.h
+++ b/Engine/Source/Runtime/CoreUObject/Public/UObject/Class.h
@@ -2045,6 +2045,14 @@ public:
 	 **/
 	void AddNativeFunction(const ANSICHAR* InName, Native InPointer);
 
+	/**
+	 * Add a native function to the internal native function table, but with a unicode name. Used when generating code from blueprints, 
+	 * which can have unicode identifiers for functions and properties.
+	 * @param	InName							name of the function
+	 * @param	InPointer						pointer to the function
+	 **/
+	void AddNativeFunction(const WIDECHAR* InName, Native InPointer);
+
 	// Add a function to the function map
 	void AddFunctionToFunctionMap(UFunction* NewFunction)
 	{
diff --git a/Engine/Source/Runtime/CoreUObject/Public/UObject/CoreNative.h b/Engine/Source/Runtime/CoreUObject/Public/UObject/CoreNative.h
--- a/Engine/Source/Runtime/CoreUObject/Public/UObject/CoreNative.h
+++ b/Engine/Source/Runtime/CoreUObject/Public/UObject/CoreNative.h
@@ -19,6 +19,8 @@ struct FNativeFunctionRegistrar
 		RegisterFunction(Class, InName, InPointer);
 	}
 	static COREUOBJECT_API void RegisterFunction(class UClass* Class, const ANSICHAR* InName, Native InPointer);
+	// overload for types generated from blueprints, which can have unicode names:
+	static COREUOBJECT_API void RegisterFunction(class UClass* Class, const WIDECHAR* InName, Native InPointer);
 };
 
 #endif
diff --git a/Engine/Source/Runtime/CoreUObject/Public/UObject/ObjectBase.h b/Engine/Source/Runtime/CoreUObject/Public/UObject/ObjectBase.h
--- a/Engine/Source/Runtime/CoreUObject/Public/UObject/ObjectBase.h
+++ b/Engine/Source/Runtime/CoreUObject/Public/UObject/ObjectBase.h
@@ -1492,18 +1492,18 @@ public: \
 #define IMPLEMENT_CORE_INTRINSIC_CLASS(TClass, TSuperClass, InitCode) \
 	IMPLEMENT_INTRINSIC_CLASS(TClass, COREUOBJECT_API, TSuperClass, COREUOBJECT_API, InitCode)
 
-// Register a dynamic class (created at runtime, not startup).
-#define IMPLEMENT_DYNAMIC_CLASS(TClass, TClassCrc) \
+// Register a dynamic class (created at runtime, not startup). Explicit ClassName parameter because Blueprint types can have names that can't be used natively:
+#define IMPLEMENT_DYNAMIC_CLASS(TClass, ClassName, TClassCrc) \
 	UClass* TClass::GetPrivateStaticClass(const TCHAR* Package) \
 	{ \
 		UPackage* PrivateStaticClassOuter = CastChecked<UPackage>(StaticFindObjectFast(UPackage::StaticClass(), nullptr, Package)); \
-		UClass* PrivateStaticClass = Cast<UClass>(StaticFindObjectFast(UClass::StaticClass(), PrivateStaticClassOuter, (TCHAR*)TEXT(#TClass) + 1)); \
+		UClass* PrivateStaticClass = Cast<UClass>(StaticFindObjectFast(UClass::StaticClass(), PrivateStaticClassOuter, (TCHAR*)ClassName)); \
 		if (!PrivateStaticClass) \
 		{ \
 			/* this could be handled with templates, but we want it external to avoid code bloat */ \
 			GetPrivateStaticClassBody( \
 			Package, \
-			(TCHAR*)TEXT(#TClass) + 1 + ((StaticClassFlags & CLASS_Deprecated) ? 11 : 0), \
+			(TCHAR*)ClassName, \
 			PrivateStaticClass, \
 			StaticRegisterNatives##TClass, \
 			sizeof(TClass), \
diff --git a/Engine/Source/Runtime/CoreUObject/Public/UObject/Script.h b/Engine/Source/Runtime/CoreUObject/Public/UObject/Script.h
--- a/Engine/Source/Runtime/CoreUObject/Public/UObject/Script.h
+++ b/Engine/Source/Runtime/CoreUObject/Public/UObject/Script.h
@@ -324,4 +324,17 @@ private:
 	bool bOldGAllowScriptExecutionInEditor;
 };
 
+/** @return True if the char can be used in an identifier in c++ */
+COREUOBJECT_API bool IsValidCPPIdentifierChar(TCHAR Char);
+
+/** @return A string that contains only Char if Char IsValidCPPIdentifierChar, otherwise returns a corresponding sequence of valid c++ chars */
+COREUOBJECT_API FString ToValidCPPIdentifierChars(TCHAR Char);
+
+/** 
+	@param InName The string to transform
+	@param bDeprecated whether the name has been deprecated
+	@Param Prefix The prefix to be prepended to the return value, accepts nullptr or empty string
+	@return A corresponding string that contains only valid c++ characters and is prefixed with Prefix
+*/
 COREUOBJECT_API FString UnicodeToCPPIdentifier(const FString& InName, bool bDeprecated, const TCHAR* Prefix);
+
