BlueprintCompilerCppBackend - unconverted blueprints:
- support for functions and delegates
- various fixes

diff --git a/Engine/Source/Developer/BlueprintCompilerCppBackend/Private/BlueprintCompilerCppBackend.cpp b/Engine/Source/Developer/BlueprintCompilerCppBackend/Private/BlueprintCompilerCppBackend.cpp
--- a/Engine/Source/Developer/BlueprintCompilerCppBackend/Private/BlueprintCompilerCppBackend.cpp
+++ b/Engine/Source/Developer/BlueprintCompilerCppBackend/Private/BlueprintCompilerCppBackend.cpp
@@ -450,7 +450,19 @@ FString FBlueprintCompilerCppBackend::EmitCallStatmentInner(FEmitterLocalContext
 	}
 	else
 	{
-		if (bStaticCall)
+		auto FunctionOwner = Statement.FunctionToCall->GetOwnerClass();
+		auto OwnerBPGC = Cast<UBlueprintGeneratedClass>(FunctionOwner);
+		const bool bUnconvertedClass = OwnerBPGC && !EmitterContext.Dependencies.WillClassBeConverted(OwnerBPGC);
+		if (bUnconvertedClass)
+		{
+			ensure(!Statement.bIsParentContext); //unsupported yet
+			ensure(!bStaticCall); //unsupported yet
+			ensure(bCallOnDifferentObject); //unexpected
+			const FString WrapperName = FString::Printf(TEXT("FUnconvertedWrapper__%s"), *FEmitHelper::GetCppName(OwnerBPGC));
+			const FString CalledObject = bCallOnDifferentObject ? TermToText(EmitterContext, Statement.FunctionContext, false) : TEXT("this");
+			Result += FString::Printf(TEXT("%s(%s)."), *WrapperName, *CalledObject);
+		}
+		else if (bStaticCall)
 		{
 			const bool bIsCustomThunk = Statement.FunctionToCall->GetBoolMetaData(TEXT("CustomThunk")) 
 				|| Statement.FunctionToCall->HasMetaData(TEXT("CustomStructureParam")) 
@@ -513,18 +525,20 @@ FString FBlueprintCompilerCppBackend::TermToText(FEmitterLocalContext& EmitterCo
 	}
 	else
 	{
-		FString ResultPath(TEXT(""));
+		FString ContextStr;
 		if ((Term->Context != nullptr) && (Term->Context->Name != PSC_Self))
 		{
+			ensure(Term->AssociatedVarProperty);
 			const bool bFromDefaultValue = Term->Context->IsClassContextType();
 			if (bFromDefaultValue)
 			{
-				const UClass* MinimalClass = ensure(Term->AssociatedVarProperty)
+				UClass* MinimalClass = Term->AssociatedVarProperty 
 					? Term->AssociatedVarProperty->GetOwnerClass()
 					: Cast<UClass>(Term->Context->Type.PinSubCategoryObject.Get());
 				if (MinimalClass)
 				{
-					ResultPath += FString::Printf(TEXT("GetDefaultValueSafe<%s>(")
+					MinimalClass = EmitterContext.GetFirstNativeParent(MinimalClass);
+					ContextStr += FString::Printf(TEXT("GetDefaultValueSafe<%s>(")
 						, *FEmitHelper::GetCppName(MinimalClass));
 				}
 				else
@@ -533,17 +547,48 @@ FString FBlueprintCompilerCppBackend::TermToText(FEmitterLocalContext& EmitterCo
 				}
 			}
 
-			ResultPath += TermToText(EmitterContext, Term->Context, false);
-			ResultPath += bFromDefaultValue ? TEXT(")") : TEXT("");
-			if (Term->Context->IsStructContextType())
+			ContextStr += TermToText(EmitterContext, Term->Context, false);
+			ContextStr += bFromDefaultValue ? TEXT(")") : TEXT("");
+		}
+
+		FString ResultPath;
+
+		if (Term->Context && Term->Context->IsStructContextType())
+		{
+			check(Term->AssociatedVarProperty);
+			ResultPath = ContextStr + TEXT(".") + FEmitHelper::GetCppName(Term->AssociatedVarProperty);
+		}
+		else if (Term->AssociatedVarProperty)
+		{
+			auto MinimalClass = Term->AssociatedVarProperty->GetOwnerClass();
+			auto MinimalBPGC = Cast<UBlueprintGeneratedClass>(MinimalClass);
+			if (MinimalBPGC && !EmitterContext.Dependencies.WillClassBeConverted(MinimalBPGC))
 			{
-				ResultPath += TEXT(".");
+				if ((!Term->Context) || (Term->Context->Name == PSC_Self))
+				{
+					ensure(ContextStr.IsEmpty());
+					ContextStr = TEXT("this");
+				}
+
+				ResultPath = FString::Printf(TEXT("FUnconvertedWrapper__%s(%s).GetRef__%s()")
+					, *FEmitHelper::GetCppName(MinimalBPGC)
+					, *ContextStr
+					, *UnicodeToCPPIdentifier(Term->AssociatedVarProperty->GetName(), false, nullptr));
 			}
 			else
 			{
-				ResultPath += TEXT("->");
+				if ((Term->Context != nullptr) && (Term->Context->Name != PSC_Self))
+				{
+					ResultPath = ContextStr + TEXT("->");
+				}
+				ResultPath += FEmitHelper::GetCppName(Term->AssociatedVarProperty);
 			}
 		}
+		else
+		{
+			ensure(ContextStr.IsEmpty());
+			ResultPath += Term->Name;
+		}
 
 		FString Conditions;
 		if (bUseSafeContext)
@@ -551,7 +596,6 @@ FString FBlueprintCompilerCppBackend::TermToText(FEmitterLocalContext& EmitterCo
 			Conditions = FSafeContextScopedEmmitter::ValidationChain(EmitterContext, Term->Context, *this);
 		}
 
-		ResultPath += Term->AssociatedVarProperty ? FEmitHelper::GetCppName(Term->AssociatedVarProperty) : Term->Name;
 		return Conditions.IsEmpty()
 			? ResultPath
 			: FString::Printf(TEXT("((%s) ? (%s) : (%s))"), *Conditions, *ResultPath, *FEmitHelper::DefaultValue(EmitterContext, Term->Type));
diff --git a/Engine/Source/Developer/BlueprintCompilerCppBackend/Private/BlueprintCompilerCppBackendBase.cpp b/Engine/Source/Developer/BlueprintCompilerCppBackend/Private/BlueprintCompilerCppBackendBase.cpp
--- a/Engine/Source/Developer/BlueprintCompilerCppBackend/Private/BlueprintCompilerCppBackendBase.cpp
+++ b/Engine/Source/Developer/BlueprintCompilerCppBackend/Private/BlueprintCompilerCppBackendBase.cpp
@@ -36,6 +36,7 @@ void FBlueprintCompilerCppBackendBase::DeclareDelegates(FEmitterLocalContext Emi
 {
 	// MC DELEGATE DECLARATION
 	{
+		Emit(Header, TEXT("\n"));
 		auto DelegateDeclarations = FEmitHelper::EmitMulticastDelegateDeclarations(EmitterContext, SourceClass);
 		FString AllDeclarations;
 		FEmitHelper::ArrayToString(DelegateDeclarations, AllDeclarations, TEXT(";\n"));
@@ -77,12 +78,12 @@ void FBlueprintCompilerCppBackendBase::DeclareDelegates(FEmitterLocalContext Emi
 			}
 		}
 
+		Emit(Header, TEXT("\n"));
 		auto DelegateDeclarations = FEmitHelper::EmitSinglecastDelegateDeclarations(EmitterContext, Delegates);
 		FString AllDeclarations;
 		FEmitHelper::ArrayToString(DelegateDeclarations, AllDeclarations, TEXT(";\n"));
 		if (DelegateDeclarations.Num())
 		{
-			Emit(Header, TEXT("\t"));
 			Emit(Header, *AllDeclarations);
 			Emit(Header, TEXT(";\n"));
 		}
@@ -425,7 +426,163 @@ void FBlueprintCompilerCppBackendBase::GenerateCodeFromStruct(UUserDefinedStruct
 	Emit(Header, TEXT("};\n"));
 }
 
-void FBlueprintCompilerCppBackendBase::EmitFileBeginning(const FString& CleanName, FGatherConvertedClassDependencies* Dependencies)
+void FBlueprintCompilerCppBackendBase::GenerateWrapperForClass(UClass* SourceClass)
+{
+	FGatherConvertedClassDependencies Dependencies(SourceClass);
+	FEmitterLocalContext EmitterContext(Dependencies);
+
+	TArray<UFunction*> FunctionsToGenerate;
+	for (auto Func : TFieldRange<UFunction>(SourceClass, EFieldIteratorFlags::ExcludeSuper))
+	{
+		// Exclude native events.. Unexpected.
+		// Exclude delegate signatures.
+		static const FName __UCSName(TEXT("UserConstructionScript"));
+		if (__UCSName == Func->GetFName())
+		{
+			continue;
+		}
+
+		FunctionsToGenerate.Add(Func);
+	}
+
+	TArray<UMulticastDelegateProperty*> MCDelegateProperties;
+	for (auto MCDelegateProp : TFieldRange<UMulticastDelegateProperty>(SourceClass, EFieldIteratorFlags::ExcludeSuper))
+	{
+		MCDelegateProperties.Add(MCDelegateProp);
+	}
+
+	const bool bGenerateAnyFunction = (0 != FunctionsToGenerate.Num()) || (0 != MCDelegateProperties.Num());
+
+	// Include standard stuff
+	EmitFileBeginning(FEmitHelper::GetBaseFilename(SourceClass), &Dependencies, bGenerateAnyFunction, true);
+
+	// DELEGATES
+	const FString DelegatesClassName = FString::Printf(TEXT("U__Delegates__%s"), *FEmitHelper::GetCppName(SourceClass));
+	auto GenerateDelegateTypeName = [](UFunction* Func) -> FString
+	{
+		return FString::Printf(TEXT("F__Delegate__%s"), *FEmitHelper::GetCppName(Func));
+	};
+	auto GenerateMulticastDelegateTypeName = [](UMulticastDelegateProperty* MCDelegateProp) -> FString
+	{
+		return FString::Printf(TEXT("F__MulticastDelegate__%s"), *FEmitHelper::GetCppName(MCDelegateProp));
+	};
+	if (bGenerateAnyFunction)
+	{
+		Emit(Header, *FString::Printf(TEXT("\nUCLASS()\nclass %s : public UObject\n{\npublic:\n\tGENERATED_BODY()\n"), *DelegatesClassName));
+		for (auto Func : FunctionsToGenerate)
+		{
+			FString ParamNumberStr, Parameters;
+			FEmitHelper::ParseDelegateDetails(EmitterContext, Func, Parameters, ParamNumberStr);
+			Emit(Header, *FString::Printf(TEXT("\tDECLARE_DYNAMIC_DELEGATE%s(%s%s);\n"), *ParamNumberStr, *GenerateDelegateTypeName(Func), *Parameters));
+		}
+		for (auto MCDelegateProp : MCDelegateProperties)
+		{
+			FString ParamNumberStr, Parameters;
+			FEmitHelper::ParseDelegateDetails(EmitterContext, MCDelegateProp->SignatureFunction, Parameters, ParamNumberStr);
+			Emit(Header, *FString::Printf(TEXT("\tDECLARE_DYNAMIC_MULTICAST_DELEGATE%s(%s%s);\n"), *ParamNumberStr, *GenerateMulticastDelegateTypeName(MCDelegateProp), *Parameters));
+		}
+		Emit(Header, TEXT("\n};\n"));
+	}
+
+	// Begin the struct
+	const FString WrapperName = FString::Printf(TEXT("FUnconvertedWrapper__%s"), *FEmitHelper::GetCppName(SourceClass));
+
+	FString ParentStruct;
+	auto SuperBPGC = Cast<UBlueprintGeneratedClass>(SourceClass->GetSuperClass());
+	if (SuperBPGC && !Dependencies.WillClassBeConverted(SuperBPGC))
+	{
+		//TODO: include header for super-class wrapper
+		ParentStruct = FString::Printf(TEXT("FUnconvertedWrapper__%s"), *FEmitHelper::GetCppName(SourceClass->GetSuperClass()));
+	}
+	else
+	{
+		ParentStruct = FString::Printf(TEXT("FUnconvertedWrapper<%s>"), *FEmitHelper::GetCppName(SourceClass->GetSuperClass()));
+	}
+	
+	Emit(Header, *FString::Printf(TEXT("struct %s : public %s\n{\n"), *WrapperName, *ParentStruct));
+	// Constructor
+	Emit(Header, *FString::Printf(TEXT("\n\t%s(const %s* __InObject) : %s(__InObject){}\n"), *WrapperName, *FEmitHelper::GetCppName(EmitterContext.GetNativeOrConvertedClass(SourceClass)), *ParentStruct));
+
+	// PROPERTIES:
+	for (auto Property : TFieldRange<UProperty>(SourceClass, EFieldIteratorFlags::ExcludeSuper))
+	{
+		//TODO: check if the property is really used?
+		const FString TypeDeclaration = Property->IsA<UMulticastDelegateProperty>()
+			? FString::Printf(TEXT("%s::%s"), *DelegatesClassName, *GenerateMulticastDelegateTypeName(CastChecked<UMulticastDelegateProperty>(Property)))
+			: EmitterContext.ExportCppDeclaration(Property, EExportedDeclaration::Parameter, EPropertyExportCPPFlags::CPPF_CustomTypeName | EPropertyExportCPPFlags::CPPF_BlueprintCppBackend | EPropertyExportCPPFlags::CPPF_NoRef, true);
+		Emit(Header, *FString::Printf(TEXT("\t%s& GetRef__%s()"), *TypeDeclaration, *UnicodeToCPPIdentifier(Property->GetName(), false, nullptr)));
+		Emit(Header, TEXT("\n\t{\n\t\t"));
+		Emit(Header, *FString::Printf(TEXT("static const FName __PropertyName(TEXT(\"%s\"));"), *Property->GetName()));
+		Emit(Header, TEXT("\n\t\t"));
+		Emit(Header, *FString::Printf(TEXT("return *(GetClass()->FindPropertyByName(__PropertyName)->ContainerPtrToValuePtr<%s>(__Object));"),*TypeDeclaration));
+		Emit(Header, TEXT("\n\t}\n\n"));
+	}
+
+	// FUNCTIONS:
+	for (auto Func : FunctionsToGenerate)
+	{
+		Emit(Header, *FString::Printf(TEXT("\tvoid %s("), *FEmitHelper::GetCppName(Func)));
+		bool bFirst = true;
+		for (TFieldIterator<UProperty> It(Func); It; ++It)
+		{
+			UProperty* Property = *It;
+			if (!Property || !Property->HasAnyPropertyFlags(CPF_Parm))
+			{
+				continue;
+			}
+
+			if (!bFirst)
+			{
+				Emit(Header, TEXT(", "));
+			}
+			else
+			{
+				bFirst = false;
+			}
+
+			if (Property->HasAnyPropertyFlags(CPF_OutParm))
+			{
+				Emit(Header, TEXT("/*out*/ "));
+			}
+
+			const FString TypeDeclaration = EmitterContext.ExportCppDeclaration(Property, EExportedDeclaration::Parameter, EPropertyExportCPPFlags::CPPF_CustomTypeName | EPropertyExportCPPFlags::CPPF_BlueprintCppBackend);
+			Emit(Header, *TypeDeclaration);
+		}
+		Emit(Header, TEXT(")\n\t{\n"));
+		Emit(Header, *FString::Printf(TEXT("\t\t%s::%s __Delegate;\n"), *DelegatesClassName, *GenerateDelegateTypeName(Func)));
+		Emit(Header, *FString::Printf(TEXT("\t\tstatic const FName __FunctionName(TEXT(\"%s\"));\n"), *Func->GetName()));
+		Emit(Header, TEXT("\t\t__Delegate.BindUFunction(__Object, __FunctionName);\n"));
+		Emit(Header, TEXT("\t\tcheck(__Delegate.IsBound());\n"));
+		Emit(Header, TEXT("\t\t__Delegate.Execute("));
+		bFirst = true;
+		for (TFieldIterator<UProperty> It(Func); It; ++It)
+		{
+			UProperty* Property = *It;
+			if (!Property || !Property->HasAnyPropertyFlags(CPF_Parm))
+			{
+				continue;
+			}
+
+			if (!bFirst)
+			{
+				Emit(Header, TEXT(", "));
+			}
+			else
+			{
+				bFirst = false;
+			}
+
+			Emit(Header, *UnicodeToCPPIdentifier(Property->GetName(), Property->HasAnyPropertyFlags(CPF_Deprecated), TEXT("bpv__")));
+		}
+		Emit(Header, TEXT(");\n"));
+		Emit(Header, TEXT("\t}\n"));
+	}
+
+	// close struct
+	Emit(Header, TEXT("};\n"));
+}
+
+void FBlueprintCompilerCppBackendBase::EmitFileBeginning(const FString& CleanName, FGatherConvertedClassDependencies* Dependencies, bool bIncludeGeneratedH, bool bIncludeCodeHelpersInHeader)
 {
 	Emit(Header, TEXT("#pragma once\n\n"));
 
@@ -435,7 +592,14 @@ void FBlueprintCompilerCppBackendBase::EmitFileBeginning(const FString& CleanNam
 	};
 	EmitIncludeHeader(Body, FApp::GetGameName(), true);
 	EmitIncludeHeader(Body, *CleanName, true);
-	EmitIncludeHeader(Body, TEXT("GeneratedCodeHelpers"), true);
+	if (bIncludeCodeHelpersInHeader)
+	{
+		EmitIncludeHeader(Header, TEXT("GeneratedCodeHelpers"), true);
+	}
+	else
+	{
+		EmitIncludeHeader(Body, TEXT("GeneratedCodeHelpers"), true);
+	}
 
 	if (Dependencies)
 	{
@@ -507,5 +671,8 @@ void FBlueprintCompilerCppBackendBase::EmitFileBeginning(const FString& CleanNam
 		EmitInner(Body, Dependencies->IncludeInBody, TSet<UField*>());
 	}
 
-	Emit(Header, *FString::Printf(TEXT("#include \"%s.generated.h\"\n\n"), *CleanName));
+	if (bIncludeGeneratedH)
+	{
+		Emit(Header, *FString::Printf(TEXT("#include \"%s.generated.h\"\n\n"), *CleanName));
+	}
 }
diff --git a/Engine/Source/Developer/BlueprintCompilerCppBackend/Private/BlueprintCompilerCppBackendBase.h b/Engine/Source/Developer/BlueprintCompilerCppBackend/Private/BlueprintCompilerCppBackendBase.h
--- a/Engine/Source/Developer/BlueprintCompilerCppBackend/Private/BlueprintCompilerCppBackendBase.h
+++ b/Engine/Source/Developer/BlueprintCompilerCppBackend/Private/BlueprintCompilerCppBackendBase.h
@@ -45,6 +45,7 @@ public:
 	virtual void GenerateCodeFromClass(UClass* SourceClass, TIndirectArray<FKismetFunctionContext>& Functions, bool bGenerateStubsOnly) override;
 	virtual void GenerateCodeFromEnum(UUserDefinedEnum* SourceEnum) override;
 	virtual void GenerateCodeFromStruct(UUserDefinedStruct* SourceStruct) override;
+	virtual void GenerateWrapperForClass(UClass* SourceClass) override;
 
 	virtual const FString& GetBody()	const override { return Body; }
 	virtual const FString& GetHeader()	const override { return Header; }
@@ -82,7 +83,7 @@ protected:
 	/** Builds both the header declaration and body implementation of a function */
 	void ConstructFunction(FKismetFunctionContext& FunctionContext, FEmitterLocalContext& EmitterContext, bool bGenerateStubOnly);
 
-	void EmitFileBeginning(const FString& CleanName, FGatherConvertedClassDependencies* Dependencies);
+	void EmitFileBeginning(const FString& CleanName, FGatherConvertedClassDependencies* Dependencies, bool bIncludeGeneratedH = true, bool bIncludeCodeHelpersInHeader = false);
 
 	int32 StatementToStateIndex(FKismetFunctionContext& FunctionContext, FBlueprintCompiledStatement* Statement)
 	{
diff --git a/Engine/Source/Developer/BlueprintCompilerCppBackend/Private/BlueprintCompilerCppBackendGatherDependencies.cpp b/Engine/Source/Developer/BlueprintCompilerCppBackend/Private/BlueprintCompilerCppBackendGatherDependencies.cpp
--- a/Engine/Source/Developer/BlueprintCompilerCppBackend/Private/BlueprintCompilerCppBackendGatherDependencies.cpp
+++ b/Engine/Source/Developer/BlueprintCompilerCppBackend/Private/BlueprintCompilerCppBackendGatherDependencies.cpp
@@ -34,7 +34,10 @@ void FGatherConvertedClassDependencies::DependenciesForHeader()
 		{
 			if (auto ObjectProperty = Cast<const UObjectPropertyBase>(Property))
 			{
-				DeclareInHeader.Add(ObjectProperty->PropertyClass);
+				if (ShouldIncludeHeaderFor(ObjectProperty->PropertyClass))
+				{
+					DeclareInHeader.Add(ObjectProperty->PropertyClass);
+				}
 			}
 			else if (auto InterfaceProperty = Cast<const UInterfaceProperty>(Property))
 			{
@@ -94,7 +97,7 @@ void FGatherConvertedClassDependencies::HandleObjectReference(UObject*& InObject
 	}
 
 	auto ObjAsField = Cast<UField>(InObject);
-	if (ObjAsField && ShouldIncludeHeaderFor(ObjAsField) && !IncludeInHeader.Contains(ObjAsField))
+	if (ObjAsField && !IncludeInHeader.Contains(ObjAsField))
 	{
 		IncludeInBody.Add(ObjAsField);
 	}
@@ -102,7 +105,6 @@ void FGatherConvertedClassDependencies::HandleObjectReference(UObject*& InObject
 	//TODO: What About Delegates?
 	auto ObjAsBPGC = Cast<UBlueprintGeneratedClass>(InObject);
 	const bool bWillBeConvetedAsBPGC = ObjAsBPGC && WillClassBeConverted(ObjAsBPGC);
-	const bool bRemainAsUnconvertedBPGC = ObjAsBPGC && !bWillBeConvetedAsBPGC;
 	if (bWillBeConvetedAsBPGC)
 	{
 		ConvertedClasses.Add(ObjAsBPGC);
@@ -123,6 +125,14 @@ void FGatherConvertedClassDependencies::HandleObjectReference(UObject*& InObject
 	}
 	else if ((InObject->IsAsset() || ObjAsBPGC) && !InObject->IsIn(OriginalStruct))
 	{
+		// include all not converted super classes
+		for (auto SuperBPGC = ObjAsBPGC ? Cast<UBlueprintGeneratedClass>(ObjAsBPGC->GetSuperClass()) : nullptr; 
+			SuperBPGC && !WillClassBeConverted(SuperBPGC);
+			SuperBPGC = Cast<UBlueprintGeneratedClass>(SuperBPGC->GetSuperClass()))
+		{
+			Assets.AddUnique(SuperBPGC);
+		}
+
 		Assets.AddUnique(InObject);
 		return;
 	}
diff --git a/Engine/Source/Developer/BlueprintCompilerCppBackend/Private/BlueprintCompilerCppBackendUtils.cpp b/Engine/Source/Developer/BlueprintCompilerCppBackend/Private/BlueprintCompilerCppBackendUtils.cpp
--- a/Engine/Source/Developer/BlueprintCompilerCppBackend/Private/BlueprintCompilerCppBackendUtils.cpp
+++ b/Engine/Source/Developer/BlueprintCompilerCppBackend/Private/BlueprintCompilerCppBackendUtils.cpp
@@ -657,7 +657,7 @@ TArray<FString> FEmitHelper::EmitSinglecastDelegateDeclarations(FEmitterLocalCon
 			| EPropertyExportCPPFlags::CPPF_BlueprintCppBackend;
 		const FString TypeName = EmitterContext.ExportCppDeclaration(It, EExportedDeclaration::Parameter, LocalExportCPPFlags, true);
 
-		Results.Add(*FString::Printf(TEXT("DECLARE_DYNAMIC_DELEGATE%s(%s%s)"), *ParamNumberStr, *TypeName, *Parameters));
+		Results.Add(*FString::Printf(TEXT("\tDECLARE_DYNAMIC_DELEGATE%s(%s%s)"), *ParamNumberStr, *TypeName, *Parameters));
 	}
 	return Results;
 }
diff --git a/Engine/Source/Developer/BlueprintCompilerCppBackend/Private/BlueprintCompilerCppBackendValueHelper.cpp b/Engine/Source/Developer/BlueprintCompilerCppBackend/Private/BlueprintCompilerCppBackendValueHelper.cpp
--- a/Engine/Source/Developer/BlueprintCompilerCppBackend/Private/BlueprintCompilerCppBackendValueHelper.cpp
+++ b/Engine/Source/Developer/BlueprintCompilerCppBackend/Private/BlueprintCompilerCppBackendValueHelper.cpp
@@ -45,7 +45,17 @@ void FEmitDefaultValueHelper::OuterGenerate(FEmitterLocalContext& Context
 			|| (!Property->Identical_InContainer(DataContainer, OptionalDefaultDataContainer, ArrayIndex) && !IsInstancedSubobjectLambda(ArrayIndex)))
 		{
 			FString PathToMember;
-			if (Property->HasAnyPropertyFlags(CPF_NativeAccessSpecifierPrivate) || (!bAllowProtected && Property->HasAnyPropertyFlags(CPF_NativeAccessSpecifierProtected)))
+			auto PropertyOwnerAsBPGC = Cast<UBlueprintGeneratedClass>(Property->GetOwnerClass());
+			if (PropertyOwnerAsBPGC && !Context.Dependencies.WillClassBeConverted(PropertyOwnerAsBPGC))
+			{
+				ensure(EPropertyAccessOperator::None != AccessOperator);
+				const FString OperatorStr = (EPropertyAccessOperator::Dot == AccessOperator) ? TEXT("&") : TEXT("");
+				const FString ContainerStr = (EPropertyAccessOperator::None == AccessOperator) ? TEXT("this") : FString::Printf(TEXT("%s(%s)"), *OperatorStr, *OuterPath);
+
+				PathToMember = FString::Printf(TEXT("FUnconvertedWrapper__%s(%s).GetRef__%s()"), *FEmitHelper::GetCppName(PropertyOwnerAsBPGC), *ContainerStr
+					, *UnicodeToCPPIdentifier(Property->GetName(), false, nullptr));
+			}
+			else if (Property->HasAnyPropertyFlags(CPF_NativeAccessSpecifierPrivate) || (!bAllowProtected && Property->HasAnyPropertyFlags(CPF_NativeAccessSpecifierProtected)))
 			{
 				ensure(EPropertyAccessOperator::None != AccessOperator);
 				
diff --git a/Engine/Source/Developer/BlueprintNativeCodeGen/Private/NativeCodeGenerationTool.cpp b/Engine/Source/Developer/BlueprintNativeCodeGen/Private/NativeCodeGenerationTool.cpp
--- a/Engine/Source/Developer/BlueprintNativeCodeGen/Private/NativeCodeGenerationTool.cpp
+++ b/Engine/Source/Developer/BlueprintNativeCodeGen/Private/NativeCodeGenerationTool.cpp
@@ -14,6 +14,7 @@
 //#include "Editor/KismetCompiler/Public/BlueprintCompilerCppBackendInterface.h"
 #include "Developer/BlueprintCompilerCppBackend/Public/BlueprintCompilerCppBackendGatherDependencies.h"
 #include "IBlueprintCompilerCppBackendModule.h" // for GetBaseFilename()
+#include "KismetCompilerModule.h"
 
 #define LOCTEXT_NAMESPACE "NativeCodeGenerationTool"
 
@@ -41,6 +42,7 @@ struct FGeneratedCodeData
 	FString BaseFilename;
 	TWeakObjectPtr<UBlueprint> Blueprint;
 	TSet<UField*> DependentObjects;
+	TSet<UBlueprintGeneratedClass*> UnconvertedNeededClasses;
 
 	void GatherUserDefinedDependencies(UBlueprint& InBlueprint)
 	{
@@ -58,16 +60,34 @@ struct FGeneratedCodeData
 			DependentObjects.Add(Iter);
 		}
 
-		TypeDependencies.Empty();
+		if (DependentObjects.Num())
+		{
+			TypeDependencies = LOCTEXT("ConvertedDependencies", "Converted Dependencies:\n").ToString();
+		}
+		else
+		{
+			TypeDependencies = LOCTEXT("NoConvertedAssets", "No Converted Dependencies was found.\n").ToString();
+		}
+
 		for (auto Obj : DependentObjects)
 		{
 			TypeDependencies += FString::Printf(TEXT("%s \t%s\n"), *Obj->GetClass()->GetName(), *Obj->GetPathName());
 		}
 		DependentObjects.Add(InBlueprint.GeneratedClass);
 
-		if (TypeDependencies.IsEmpty())
+		bool bUnconvertedHeader = false;
+		for (auto Asset : ClassDependencies.Assets)
 		{
-			TypeDependencies += LOCTEXT("NoNonNativeDependencies", "No non-native dependencies.").ToString();
+			if (auto BPGC = Cast<UBlueprintGeneratedClass>(Asset))
+			{
+				UnconvertedNeededClasses.Add(BPGC);
+				if (!bUnconvertedHeader)
+				{
+					bUnconvertedHeader = true;
+					TypeDependencies += LOCTEXT("NoConvertedAssets", "\nUnconverted Dependencies, that require a warpper struct:\n").ToString();
+				}
+				TypeDependencies += FString::Printf(TEXT("%s \t%s\n"), *BPGC->GetClass()->GetName(), *BPGC->GetPathName());
+			}
 		}
 	}
 
@@ -101,14 +121,15 @@ struct FGeneratedCodeData
 			return false;
 		}
 
-		const int WorkParts = 4 + 2 * DependentObjects.Num();
+		const int WorkParts = 3 + (4 * DependentObjects.Num()) + (2 * UnconvertedNeededClasses.Num());
 		FScopedSlowTask SlowTask(WorkParts, LOCTEXT("GeneratingCppFiles", "Generating C++ files.."));
 		SlowTask.MakeDialog();
-		SlowTask.EnterProgressFrame();
 
 		TArray<FString> CreatedFiles;
 		for(auto Obj : DependentObjects)
 		{
+			SlowTask.EnterProgressFrame();
+
 			TSharedPtr<FString> HeaderSource(new FString());
 			TSharedPtr<FString> CppSource(new FString());
 			FBlueprintNativeCodeGenUtils::GenerateCppCode(Obj, HeaderSource, CppSource);
@@ -143,6 +164,29 @@ struct FGeneratedCodeData
 			}
 		}
 
+		for (auto BPGC : UnconvertedNeededClasses)
+		{
+			SlowTask.EnterProgressFrame();
+
+			IKismetCompilerInterface& Compiler = FModuleManager::LoadModuleChecked<IKismetCompilerInterface>(KISMET_COMPILER_MODULENAME);
+			const FString HeaderSource = Compiler.GenerateCppWrapper(BPGC);
+
+			SlowTask.EnterProgressFrame();
+
+			const FString BackendBaseFilename = IBlueprintCompilerCppBackendModule::GetBaseFilename(BPGC);
+
+			const FString FullHeaderFilename = FPaths::Combine(*HeaderDirPath, *(BackendBaseFilename + TEXT(".h")));
+			const bool bHeaderSaved = FFileHelper::SaveStringToFile(HeaderSource, *FullHeaderFilename);
+			if (!bHeaderSaved)
+			{
+				ErrorString += FString::Printf(*LOCTEXT("HeaderNotSaved", "Header file wasn't saved. Check log for details. %s\n").ToString(), *BPGC->GetPathName());
+			}
+			else
+			{
+				CreatedFiles.Add(FullHeaderFilename);
+			}
+		}
+
 		SlowTask.EnterProgressFrame();
 
 		bool bGenerateProjectFiles = true;
diff --git a/Engine/Source/Editor/KismetCompiler/Private/KismetCompilerModule.cpp b/Engine/Source/Editor/KismetCompiler/Private/KismetCompilerModule.cpp
--- a/Engine/Source/Editor/KismetCompiler/Private/KismetCompilerModule.cpp
+++ b/Engine/Source/Editor/KismetCompiler/Private/KismetCompilerModule.cpp
@@ -143,6 +143,13 @@ FString FKismet2CompilerModule::GenerateCppCodeForStruct(UUserDefinedStruct* UDS
 	return Backend_CPP->GetHeader();
 }
 
+FString FKismet2CompilerModule::GenerateCppWrapper(UBlueprintGeneratedClass* BPGC)
+{
+	TUniquePtr<IBlueprintCompilerCppBackend> Backend_CPP(IBlueprintCompilerCppBackendModuleInterface::Get().Create());
+	Backend_CPP->GenerateWrapperForClass(BPGC);
+	return Backend_CPP->GetHeader();
+}
+
 extern UNREALED_API FSecondsCounterData BlueprintCompileAndLoadTimerData;
 
 // Compiles a blueprint.
diff --git a/Engine/Source/Editor/KismetCompiler/Public/BlueprintCompilerCppBackendInterface.h b/Engine/Source/Editor/KismetCompiler/Public/BlueprintCompilerCppBackendInterface.h
--- a/Engine/Source/Editor/KismetCompiler/Public/BlueprintCompilerCppBackendInterface.h
+++ b/Engine/Source/Editor/KismetCompiler/Public/BlueprintCompilerCppBackendInterface.h
@@ -12,6 +12,9 @@ public:
 	virtual void GenerateCodeFromEnum(UUserDefinedEnum* SourceEnum) = 0;
 	virtual void GenerateCodeFromStruct(UUserDefinedStruct* SourceStruct) = 0;
 
+	// Generate a wrapper class, that helps accessing non-native properties and calling non-native functions
+	virtual void GenerateWrapperForClass(UClass* SourceClass) = 0;
+
 	virtual const FString& GetBody() const = 0;
 	virtual const FString& GetHeader() const = 0;
 
diff --git a/Engine/Source/Editor/KismetCompiler/Public/KismetCompilerModule.h b/Engine/Source/Editor/KismetCompiler/Public/KismetCompilerModule.h
--- a/Engine/Source/Editor/KismetCompiler/Public/KismetCompilerModule.h
+++ b/Engine/Source/Editor/KismetCompiler/Public/KismetCompilerModule.h
@@ -87,6 +87,8 @@ public:
 
 	virtual FString GenerateCppCodeForEnum(UUserDefinedEnum* UDEnum) = 0;
 	virtual FString GenerateCppCodeForStruct(UUserDefinedStruct* UDStruct) = 0;
+	// Generate a wrapper class, that helps accessing non-native properties and calling non-native functions
+	virtual FString GenerateCppWrapper(UBlueprintGeneratedClass* BPGC) = 0;
 };
 
 //////////////////////////////////////////////////////////////////////////
@@ -107,6 +109,7 @@ public:
 	virtual void GetBlueprintTypesForClass(UClass* ParentClass, UClass*& OutBlueprintClass, UClass*& OutBlueprintGeneratedClass) const override;
 	virtual FString GenerateCppCodeForEnum(UUserDefinedEnum* UDEnum) override;
 	virtual FString GenerateCppCodeForStruct(UUserDefinedStruct* UDStruct) override;
+	virtual FString GenerateCppWrapper(UBlueprintGeneratedClass* BPGC) override;
 	// End implementation
 private:
 	void CompileBlueprintInner(class UBlueprint* Blueprint, const FKismetCompilerOptions& CompileOptions, FCompilerResultsLog& Results, TSharedPtr<FBlueprintCompileReinstancer> Reinstancer, TArray<UObject*>* ObjLoaded);
diff --git a/Engine/Source/Runtime/CoreUObject/Private/UObject/Property.cpp b/Engine/Source/Runtime/CoreUObject/Private/UObject/Property.cpp
--- a/Engine/Source/Runtime/CoreUObject/Private/UObject/Property.cpp
+++ b/Engine/Source/Runtime/CoreUObject/Private/UObject/Property.cpp
@@ -530,7 +530,7 @@ void UProperty::ExportCppDeclaration(FOutputDevice& Out, EExportedDeclaration::T
 		}
 	}
 
-	const bool bCanHaveRef = 0 == (AdditionalExportCPPFlags & CPPF_NoConst);
+	const bool bCanHaveRef = 0 == (AdditionalExportCPPFlags & CPPF_NoRef);
 	if(auto BoolProperty = dynamic_cast<const UBoolProperty*>(this) )
 	{
 		// if this is a member variable, export it as a bitfield
diff --git a/Engine/Source/Runtime/CoreUObject/Private/UObject/PropertyDelegate.cpp b/Engine/Source/Runtime/CoreUObject/Private/UObject/PropertyDelegate.cpp
--- a/Engine/Source/Runtime/CoreUObject/Private/UObject/PropertyDelegate.cpp
+++ b/Engine/Source/Runtime/CoreUObject/Private/UObject/PropertyDelegate.cpp
@@ -75,6 +75,13 @@ bool UDelegateProperty::NetSerializeItem( FArchive& Ar, UPackageMap* Map, void*
 FString UDelegateProperty::GetCPPType( FString* ExtendedTypeText/*=NULL*/, uint32 CPPExportFlags/*=0*/ ) const
 {
 	FString UnmangledFunctionName = SignatureFunction->GetName().LeftChop( FString( HEADER_GENERATED_DELEGATE_SIGNATURE_SUFFIX ).Len() );
+	if ((0 != (CPPExportFlags & EPropertyExportCPPFlags::CPPF_BlueprintCppBackend)) && !SignatureFunction->GetOwnerClass()->HasAnyClassFlags(CLASS_Native))
+	{
+		// the name must be unique
+		const FString OwnerName = UnicodeToCPPIdentifier(SignatureFunction->GetOwnerClass()->GetName(), false, TEXT(""));
+		const FString NewUnmangledFunctionName = FString::Printf(TEXT("%s__%s"), *UnmangledFunctionName, *OwnerName);
+		UnmangledFunctionName = NewUnmangledFunctionName;
+	}
 	if (0 != (CPPExportFlags & EPropertyExportCPPFlags::CPPF_CustomTypeName))
 	{
 		UnmangledFunctionName += TEXT("__SinglecastDelegate");
diff --git a/Engine/Source/Runtime/CoreUObject/Private/UObject/PropertyMulticastDelegate.cpp b/Engine/Source/Runtime/CoreUObject/Private/UObject/PropertyMulticastDelegate.cpp
--- a/Engine/Source/Runtime/CoreUObject/Private/UObject/PropertyMulticastDelegate.cpp
+++ b/Engine/Source/Runtime/CoreUObject/Private/UObject/PropertyMulticastDelegate.cpp
@@ -123,6 +123,13 @@ FString UMulticastDelegateProperty::GetCPPType( FString* ExtendedTypeText/*=NULL
 #endif
 
 	FString UnmangledFunctionName = SignatureFunction->GetName().LeftChop( FString( HEADER_GENERATED_DELEGATE_SIGNATURE_SUFFIX ).Len() );
+	if ((0 != (CPPExportFlags & EPropertyExportCPPFlags::CPPF_BlueprintCppBackend)) && !SignatureFunction->GetOwnerClass()->HasAnyClassFlags(CLASS_Native))
+	{
+		// the name must be unique
+		const FString OwnerName = UnicodeToCPPIdentifier(SignatureFunction->GetOwnerClass()->GetName(), false, TEXT(""));
+		const FString NewUnmangledFunctionName = FString::Printf(TEXT("%s__%s"), *UnmangledFunctionName, *OwnerName);
+		UnmangledFunctionName = NewUnmangledFunctionName;
+	}
 	if (0 != (CPPExportFlags & EPropertyExportCPPFlags::CPPF_CustomTypeName))
 	{
 		UnmangledFunctionName += TEXT("__MulticastDelegate");
diff --git a/Engine/Source/Runtime/Engine/Public/GeneratedCodeHelpers.h b/Engine/Source/Runtime/Engine/Public/GeneratedCodeHelpers.h
--- a/Engine/Source/Runtime/Engine/Public/GeneratedCodeHelpers.h
+++ b/Engine/Source/Runtime/Engine/Public/GeneratedCodeHelpers.h
@@ -317,4 +317,21 @@ ValueType& TSwitchValue(const IndexType& CurrentIndex, ValueType& DefaultValue,
 
 	return SelectedValuePtr ? *SelectedValuePtr : DefaultValue;
 }
-#endif //PLATFORM_COMPILER_HAS_VARIADIC_TEMPLATES
\ No newline at end of file
+#endif //PLATFORM_COMPILER_HAS_VARIADIC_TEMPLATES
+
+// Base class for wrappers for unconverted BlueprintGeneratedClasses
+template<class NativeType>
+struct FUnconvertedWrapper
+{
+	NativeType* __Object;
+
+	FUnconvertedWrapper(const UObject* InObject) : __Object(CastChecked<NativeType>(const_cast<UObject*>(InObject))) {}
+
+	operator NativeType*() const { return __Object; }
+
+	UClass* GetClass() const
+	{
+		check(__Object);
+		return __Object->GetClass();
+	}
+};
\ No newline at end of file
