BlueprintCompilerCppBackend:
- Huge refectoring - no more Emit. FEmitterLocalContext is used everywhere.
- No more AddHighPriorityLine
- Various minor bugs were fixed.

diff --git a/Engine/Source/Developer/BlueprintCompilerCppBackend/Private/BlueprintCompilerCppBackend.cpp b/Engine/Source/Developer/BlueprintCompilerCppBackend/Private/BlueprintCompilerCppBackend.cpp
--- a/Engine/Source/Developer/BlueprintCompilerCppBackend/Private/BlueprintCompilerCppBackend.cpp
+++ b/Engine/Source/Developer/BlueprintCompilerCppBackend/Private/BlueprintCompilerCppBackend.cpp
@@ -161,7 +161,7 @@ void FBlueprintCompilerCppBackend::EmitDynamicCastStatement(FEmitterLocalContext
 	auto BPGC = Cast<UBlueprintGeneratedClass>(ClassPtr);
 	if (BPGC && !EmitterContext.Dependencies.WillClassBeConverted(BPGC))
 	{
-		const FString NativeClass = FEmitHelper::GetCppName(EmitterContext.GetFirstNativeParent(ClassPtr));
+		const FString NativeClass = FEmitHelper::GetCppName(EmitterContext.GetFirstNativeOrConvertedClass(ClassPtr));
 		const FString TargetClass = EmitterContext.FindGloballyMappedObject(ClassPtr, UClass::StaticClass(), true);
 		EmitterContext.AddLine(FString::Printf(TEXT("%s = NoNativeCast<%s>(%s, %s);"), *CastedValue, *NativeClass, *TargetClass, *ObjectValue));
 	}
@@ -537,7 +537,7 @@ FString FBlueprintCompilerCppBackend::TermToText(FEmitterLocalContext& EmitterCo
 					: Cast<UClass>(Term->Context->Type.PinSubCategoryObject.Get());
 				if (MinimalClass)
 				{
-					MinimalClass = EmitterContext.GetFirstNativeParent(MinimalClass);
+					MinimalClass = EmitterContext.GetFirstNativeOrConvertedClass(MinimalClass);
 					ContextStr += FString::Printf(TEXT("GetDefaultValueSafe<%s>(")
 						, *FEmitHelper::GetCppName(MinimalClass));
 				}
@@ -631,9 +631,8 @@ FString FBlueprintCompilerCppBackend::LatentFunctionInfoTermToText(FEmitterLocal
 	return FEmitHelper::LiteralTerm(EmitterContext, Term->Type, StructValues, nullptr);
 }
 
-FString FBlueprintCompilerCppBackend::InnerFunctionImplementation(FKismetFunctionContext& FunctionContext, FEmitterLocalContext& EmitterContext, bool bUseSwitchState)
+void FBlueprintCompilerCppBackend::InnerFunctionImplementation(FKismetFunctionContext& FunctionContext, FEmitterLocalContext& EmitterContext, bool bUseSwitchState)
 {
-	EmitterContext.IncreaseIndent();
 	if (bUseSwitchState)
 	{
 		if (FunctionContext.bUseFlowStack)
@@ -786,10 +785,4 @@ FString FBlueprintCompilerCppBackend::InnerFunctionImplementation(FKismetFunctio
 		EmitterContext.DecreaseIndent();
 		EmitterContext.AddLine(TEXT("} while( CurrentState != -1 );"));
 	}
-
-
-	const FString Result = EmitterContext.GetResult();
-	EmitterContext.ClearResult();
-	EmitterContext.DecreaseIndent();
-	return Result;
 }
diff --git a/Engine/Source/Developer/BlueprintCompilerCppBackend/Private/BlueprintCompilerCppBackend.h b/Engine/Source/Developer/BlueprintCompilerCppBackend/Private/BlueprintCompilerCppBackend.h
--- a/Engine/Source/Developer/BlueprintCompilerCppBackend/Private/BlueprintCompilerCppBackend.h
+++ b/Engine/Source/Developer/BlueprintCompilerCppBackend/Private/BlueprintCompilerCppBackend.h
@@ -9,7 +9,7 @@ struct FEmitterLocalContext;
 class FBlueprintCompilerCppBackend : public FBlueprintCompilerCppBackendBase
 {
 protected:
-	virtual FString InnerFunctionImplementation(FKismetFunctionContext& FunctionContext, FEmitterLocalContext& EmitterContext, bool bUseSwitchState) override;
+	virtual void InnerFunctionImplementation(FKismetFunctionContext& FunctionContext, FEmitterLocalContext& EmitterContext, bool bUseSwitchState) override;
 
 protected:
 	void EmitCallStatment(FEmitterLocalContext& EmitterContext, FKismetFunctionContext& FunctionContext, FBlueprintCompiledStatement& Statement);
diff --git a/Engine/Source/Developer/BlueprintCompilerCppBackend/Private/BlueprintCompilerCppBackendBase.cpp b/Engine/Source/Developer/BlueprintCompilerCppBackend/Private/BlueprintCompilerCppBackendBase.cpp
--- a/Engine/Source/Developer/BlueprintCompilerCppBackend/Private/BlueprintCompilerCppBackendBase.cpp
+++ b/Engine/Source/Developer/BlueprintCompilerCppBackend/Private/BlueprintCompilerCppBackendBase.cpp
@@ -5,15 +5,14 @@
 #include "BlueprintCompilerCppBackendUtils.h"
 #include "IBlueprintCompilerCppBackendModule.h" // for GetBaseFilename()
 
-void FBlueprintCompilerCppBackendBase::EmitStructProperties(FEmitterLocalContext EmitterContext, UStruct* SourceClass)
+void FBlueprintCompilerCppBackendBase::EmitStructProperties(FEmitterLocalContext& EmitterContext, UStruct* SourceClass)
 {
 	// Emit class variables
 	for (TFieldIterator<UProperty> It(SourceClass, EFieldIteratorFlags::ExcludeSuper); It; ++It)
 	{
 		UProperty* Property = *It;
 		check(Property);
-
-		Emit(Header, TEXT("\n\tUPROPERTY("));
+		FString PropertyMacro(TEXT("UPROPERTY("));
 		{
 			TArray<FString> Tags = FEmitHelper::ProperyFlagsToTags(Property->PropertyFlags, nullptr != Cast<UClass>(SourceClass));
 			Tags.Emplace(FEmitHelper::HandleRepNotifyFunc(Property));
@@ -22,35 +21,27 @@ void FBlueprintCompilerCppBackendBase::EmitStructProperties(FEmitterLocalContext
 
 			FString AllTags;
 			FEmitHelper::ArrayToString(Tags, AllTags, TEXT(", "));
-			Emit(Header, *AllTags);
+			PropertyMacro += AllTags;
 		}
-		Emit(Header, TEXT(")\n"));
-		Emit(Header, TEXT("\t"));
+		PropertyMacro += TEXT(")");
+		EmitterContext.Header.AddLine(PropertyMacro);
+
 		const FString CppDeclaration = EmitterContext.ExportCppDeclaration(Property, EExportedDeclaration::Member, EPropertyExportCPPFlags::CPPF_CustomTypeName | EPropertyExportCPPFlags::CPPF_BlueprintCppBackend);
-		Emit(Header, *CppDeclaration);
-		Emit(Header, TEXT(";\n"));
+		EmitterContext.Header.AddLine(CppDeclaration + TEXT(";"));
 	}
 }
 
-void FBlueprintCompilerCppBackendBase::DeclareDelegates(FEmitterLocalContext EmitterContext, UClass* SourceClass, TIndirectArray<FKismetFunctionContext>& Functions)
+void FBlueprintCompilerCppBackendBase::DeclareDelegates(FEmitterLocalContext& EmitterContext, TIndirectArray<FKismetFunctionContext>& Functions)
 {
 	// MC DELEGATE DECLARATION
 	{
-		Emit(Header, TEXT("\n"));
-		auto DelegateDeclarations = FEmitHelper::EmitMulticastDelegateDeclarations(EmitterContext, SourceClass);
-		FString AllDeclarations;
-		FEmitHelper::ArrayToString(DelegateDeclarations, AllDeclarations, TEXT(";\n"));
-		if (DelegateDeclarations.Num())
-		{
-			Emit(Header, *AllDeclarations);
-			Emit(Header, TEXT(";\n"));
-		}
+		FEmitHelper::EmitMulticastDelegateDeclarations(EmitterContext);
 	}
 
 	// GATHER ALL SC DELEGATES
 	{
 		TArray<UDelegateProperty*> Delegates;
-		for (TFieldIterator<UDelegateProperty> It(SourceClass, EFieldIteratorFlags::ExcludeSuper); It; ++It)
+		for (TFieldIterator<UDelegateProperty> It(EmitterContext.GetCurrentlyGeneratedClass(), EFieldIteratorFlags::ExcludeSuper); It; ++It)
 		{
 			Delegates.Add(*It);
 		}
@@ -78,86 +69,79 @@ void FBlueprintCompilerCppBackendBase::DeclareDelegates(FEmitterLocalContext Emi
 			}
 		}
 
-		Emit(Header, TEXT("\n"));
-		auto DelegateDeclarations = FEmitHelper::EmitSinglecastDelegateDeclarations(EmitterContext, Delegates);
-		FString AllDeclarations;
-		FEmitHelper::ArrayToString(DelegateDeclarations, AllDeclarations, TEXT(";\n"));
-		if (DelegateDeclarations.Num())
-		{
-			Emit(Header, *AllDeclarations);
-			Emit(Header, TEXT(";\n"));
-		}
+		FEmitHelper::EmitSinglecastDelegateDeclarations(EmitterContext, Delegates);
 	}
 }
 
-void FBlueprintCompilerCppBackendBase::GenerateCodeFromClass(UClass* SourceClass, TIndirectArray<FKismetFunctionContext>& Functions, bool bGenerateStubsOnly)
+FString FBlueprintCompilerCppBackendBase::GenerateCodeFromClass(UClass* SourceClass, TIndirectArray<FKismetFunctionContext>& Functions, bool bGenerateStubsOnly, FString& OutCppBody)
 {
+	CleanBackend();
+
 	// use GetBaseFilename() so that we can coordinate #includes and filenames
 	auto CleanCppClassName = FEmitHelper::GetBaseFilename(SourceClass);
-	CppClassName = FEmitHelper::GetCppName(SourceClass);
+	auto CppClassName = FEmitHelper::GetCppName(SourceClass);
 	
 	FGatherConvertedClassDependencies Dependencies(SourceClass);
 	FEmitterLocalContext EmitterContext(Dependencies);
-	EmitFileBeginning(CleanCppClassName, &Dependencies);
+	EmitFileBeginning(CleanCppClassName, EmitterContext);
 
 	// Class declaration
 	const bool bIsInterface = SourceClass->IsChildOf<UInterface>();
 	if (bIsInterface)
 	{
-		Emit(Header, TEXT("UINTERFACE(Blueprintable"));
-		EmitReplaceConvertedMetaData(SourceClass);
-		Emit(Header, *FString::Printf(TEXT(")\nclass %s : public UInterface\n{\n\tGENERATED_BODY()\n};\n"), *FEmitHelper::GetCppName(SourceClass, true)));
-		Emit(Header, *FString::Printf(TEXT("\nclass %s"), *CppClassName));
+		EmitterContext.Header.AddLine(FString::Printf(TEXT("UINTERFACE(Blueprintable %s)"), *FEmitHelper::ReplaceConvertedMetaData(SourceClass)));
+		EmitterContext.Header.AddLine(FString::Printf(TEXT("class %s : public UInterface"), *FEmitHelper::GetCppName(SourceClass, true)));
+		EmitterContext.Header.AddLine(TEXT("{"));
+		EmitterContext.Header.IncreaseIndent();
+		EmitterContext.Header.AddLine(TEXT("GENERATED_BODY()"));
+		EmitterContext.Header.DecreaseIndent();
+		EmitterContext.Header.AddLine(TEXT("};"));
+		EmitterContext.Header.AddLine(FString::Printf(TEXT("class %s"), *CppClassName));
 	}
 	else
 	{
-		Emit(Header, TEXT("UCLASS("));
-		if (!SourceClass->IsChildOf<UBlueprintFunctionLibrary>())
-		{
-			Emit(Header, TEXT("Blueprintable, BlueprintType"));
-		}
-		EmitReplaceConvertedMetaData(SourceClass);
-		Emit(Header, TEXT(")\n"));
+		EmitterContext.Header.AddLine(FString::Printf(TEXT("UCLASS(%s%s)")
+			, (!SourceClass->IsChildOf<UBlueprintFunctionLibrary>()) ? TEXT("Blueprintable, BlueprintType") : TEXT("")
+			, *FEmitHelper::ReplaceConvertedMetaData(SourceClass)));
 
 		UClass* SuperClass = SourceClass->GetSuperClass();
-		Emit(Header, *FString::Printf(TEXT("class %s : public %s"), *CppClassName, *FEmitHelper::GetCppName(SuperClass)));
+		FString ClassDefinition = FString::Printf(TEXT("class %s : public %s"), *CppClassName, *FEmitHelper::GetCppName(SuperClass));
 
 		for (auto& ImplementedInterface : SourceClass->Interfaces)
 		{
 			if (ImplementedInterface.Class)
 			{
-				Emit(Header, *FString::Printf(TEXT(", public %s"), *FEmitHelper::GetCppName(ImplementedInterface.Class)));
+				ClassDefinition += FString::Printf(TEXT(", public %s"), *FEmitHelper::GetCppName(ImplementedInterface.Class));
 			}
 		}
+		EmitterContext.Header.AddLine(ClassDefinition);
 	}
-
+	
 	// Begin scope
-	Emit(Header,TEXT("\n{\npublic:\n\tGENERATED_BODY()\n"));
+	EmitterContext.Header.AddLine(TEXT("{"));
+	EmitterContext.Header.AddLine(TEXT("public:"));
+	EmitterContext.Header.IncreaseIndent();
+	EmitterContext.Header.AddLine(TEXT("GENERATED_BODY()"));
 
-	DeclareDelegates(EmitterContext, SourceClass, Functions);
+	DeclareDelegates(EmitterContext, Functions);
 
 	EmitStructProperties(EmitterContext, SourceClass);
 
-	// Create the state map
-	for (int32 i = 0; i < Functions.Num(); ++i)
-	{
-		StateMapPerFunction.Add(FFunctionLabelInfo());
-		FunctionIndexMap.Add(&Functions[i], i);
-	}
-
 	// Emit function declarations and definitions (writes to header and body simultaneously)
-	if (Functions.Num() > 0)
+	if (!bIsInterface)
 	{
-		Emit(Header, TEXT("\n"));
+		EmitterContext.Header.AddLine(FString::Printf(TEXT("%s(const FObjectInitializer& ObjectInitializer = FObjectInitializer::Get());"), *CppClassName));
+		EmitterContext.Header.AddLine(TEXT("virtual void PostLoadSubobjects(FObjectInstancingGraph* OuterInstanceGraph) override;"));
+		EmitterContext.Header.AddLine(TEXT("static void __StaticDependenciesAssets(TArray<FName>& OutPackagePaths);"));
+
+		FEmitDefaultValueHelper::GenerateConstructor(EmitterContext);
 	}
 
-	if (!bIsInterface)
+	// Create the state map
+	for (int32 i = 0; i < Functions.Num(); ++i)
 	{
-		Emit(Header, *FString::Printf(TEXT("\t%s(const FObjectInitializer& ObjectInitializer = FObjectInitializer::Get());\n\n"), *CppClassName));
-		Emit(Body, *FEmitDefaultValueHelper::GenerateConstructor(EmitterContext));
-
-		Emit(Header, TEXT("\n\tvirtual void PostLoadSubobjects(FObjectInstancingGraph* OuterInstanceGraph) override;\n"));
-		Emit(Header, TEXT("\n\tstatic void __StaticDependenciesAssets(TArray<FName>& OutPackagePaths);\n"));
+		StateMapPerFunction.Add(FFunctionLabelInfo());
+		FunctionIndexMap.Add(&Functions[i], i);
 	}
 
 	for (int32 i = 0; i < Functions.Num(); ++i)
@@ -168,28 +152,31 @@ void FBlueprintCompilerCppBackendBase::GenerateCodeFromClass(UClass* SourceClass
 		}
 	}
 
-	Emit(Header, *FBackendHelperUMG::WidgetFunctionsInHeader(SourceClass));
+	EmitterContext.Header.DecreaseIndent();
+	EmitterContext.Header.AddLine(TEXT("public:"));
+	EmitterContext.Header.IncreaseIndent();
+
+	FBackendHelperUMG::WidgetFunctionsInHeader(EmitterContext);
 
-	Emit(Header, TEXT("};\n\n"));
+	EmitterContext.Header.DecreaseIndent();
+	EmitterContext.Header.AddLine(TEXT("};"));
 
-	Emit(Body, *FEmitHelper::EmitLifetimeReplicatedPropsImpl(SourceClass, CppClassName, TEXT("")));
+	FEmitHelper::EmitLifetimeReplicatedPropsImpl(EmitterContext);
+
+	CleanBackend();
+
+	OutCppBody = EmitterContext.Body.Result;
+	return EmitterContext.Header.Result;
 }
 
-void FBlueprintCompilerCppBackendBase::DeclareLocalVariables(FKismetFunctionContext& FunctionContext, FEmitterLocalContext& EmitterContext, TArray<UProperty*>& LocalVariables)
+void FBlueprintCompilerCppBackendBase::DeclareLocalVariables(FEmitterLocalContext& EmitterContext, TArray<UProperty*>& LocalVariables)
 {
 	for (int32 i = 0; i < LocalVariables.Num(); ++i)
 	{
 		UProperty* LocalVariable = LocalVariables[i];
 
-		Emit(Body, TEXT("\t"));
 		const FString CppDeclaration = EmitterContext.ExportCppDeclaration(LocalVariable, EExportedDeclaration::Local, EPropertyExportCPPFlags::CPPF_CustomTypeName | EPropertyExportCPPFlags::CPPF_BlueprintCppBackend);
-		Emit(Body, *CppDeclaration);
-		Emit(Body, TEXT("{};\n"));
-	}
-
-	if (LocalVariables.Num() > 0)
-	{
-		Emit(Body, TEXT("\n"));
+		EmitterContext.AddLine(CppDeclaration + TEXT("{};"));
 	}
 }
 
@@ -204,7 +191,6 @@ void FBlueprintCompilerCppBackendBase::ConstructFunction(FKismetFunctionContext&
 
 	UProperty* ReturnValue = NULL;
 	TArray<UProperty*> LocalVariables;
-
 	{
 		FString FunctionName = FEmitHelper::GetCppName(Function);
 
@@ -271,7 +257,7 @@ void FBlueprintCompilerCppBackendBase::ConstructFunction(FKismetFunctionContext&
 
 		//@TODO: Make the header+body export more uniform
 		{
-			const FString Start = FString::Printf(TEXT("%s %s%s%s("), *ReturnType, TEXT("%s"), TEXT("%s"), *FunctionName);
+			const FString Start = FString::Printf(TEXT("%s %s%s%s"), *ReturnType, TEXT("%s"), TEXT("%s"), *FunctionName);
 
 			TArray<FString> AdditionalMetaData;
 			if (bAddCppFromBpEventMD)
@@ -281,61 +267,61 @@ void FBlueprintCompilerCppBackendBase::ConstructFunction(FKismetFunctionContext&
 
 			if (!bGenerateAsNativeEventImplementation)
 			{
-				Emit(Header, *FString::Printf(TEXT("\t%s\n"), *FEmitHelper::EmitUFuntion(Function, &AdditionalMetaData)));
+				EmitterContext.Header.AddLine(FString::Printf(TEXT("%s"), *FEmitHelper::EmitUFuntion(Function, &AdditionalMetaData)));
 			}
-			Emit(Header, TEXT("\t"));
 
+			FString HeaderFunction;
 			if (Function->HasAllFunctionFlags(FUNC_Static))
 			{
-				Emit(Header, TEXT("static "));
+				HeaderFunction += TEXT("static ");
 			}
 
 			if (bGenerateAsNativeEventImplementation)
 			{
-				Emit(Header, TEXT("virtual "));
+				HeaderFunction += TEXT("virtual ");
 			}
 
-			Emit(Header, *FString::Printf(*Start, TEXT(""), TEXT("")));
-			Emit(Body, *FString::Printf(*Start, *CppClassName, TEXT("::")));
+			HeaderFunction += FString::Printf(*Start, TEXT(""), TEXT(""));
+			FString BodyFunction = FString::Printf(*Start, *FEmitHelper::GetCppName(EmitterContext.GetCurrentlyGeneratedClass()), TEXT("::"));
 
+			FString ArgListStr = TEXT("(");
 			for (int32 i = 0; i < ArgumentList.Num(); ++i)
 			{
 				UProperty* ArgProperty = ArgumentList[i];
 
 				if (i > 0)
 				{
-					Emit(Header, TEXT(", "));
-					Emit(Body, TEXT(", "));
+					ArgListStr += TEXT(", ");
 				}
 
 				if (ArgProperty->HasAnyPropertyFlags(CPF_OutParm))
 				{
-					Emit(Header, TEXT("/*out*/ "));
-					Emit(Body, TEXT("/*out*/ "));
+					ArgListStr += TEXT("/*out*/ ");
 				}
 
-				const FString TypeDeclaration = EmitterContext.ExportCppDeclaration(ArgProperty, EExportedDeclaration::Parameter, EPropertyExportCPPFlags::CPPF_CustomTypeName | EPropertyExportCPPFlags::CPPF_BlueprintCppBackend);
-				Emit(Header, *TypeDeclaration);
-				Emit(Body, *TypeDeclaration);
+				ArgListStr += EmitterContext.ExportCppDeclaration(ArgProperty, EExportedDeclaration::Parameter, EPropertyExportCPPFlags::CPPF_CustomTypeName | EPropertyExportCPPFlags::CPPF_BlueprintCppBackend);
 			}
-
-			Emit(Header, TEXT(")"));
+			ArgListStr += TEXT(")");
+			HeaderFunction += ArgListStr;
+			BodyFunction += ArgListStr;
 			if (bGenerateAsNativeEventImplementation)
 			{
-				Emit(Header, TEXT(" override"));
+				HeaderFunction += TEXT(" override");
 			}
-			Emit(Header, TEXT(";\n"));
-			Emit(Body, TEXT(")\n"));
-		}
+			HeaderFunction += TEXT(";");
 
-		// Start the body of the implementation
-		Emit(Body, TEXT("{\n"));
+			EmitterContext.AddLine(BodyFunction);
+			EmitterContext.Header.AddLine(HeaderFunction);
+		}
 	}
 
+	// Start the body of the implementation
+	EmitterContext.AddLine(TEXT("{"));
+	EmitterContext.IncreaseIndent();
 	if (!bGenerateStubOnly)
 	{
 		// Emit local variables
-		DeclareLocalVariables(FunctionContext, EmitterContext, LocalVariables);
+		DeclareLocalVariables(EmitterContext, LocalVariables);
 
 		const bool bUseSwitchState = FunctionContext.MustUseSwitchState(nullptr);
 
@@ -361,27 +347,26 @@ void FBlueprintCompilerCppBackendBase::ConstructFunction(FKismetFunctionContext&
 			}
 		}
 
-		const FString FunctionImplementation = InnerFunctionImplementation(FunctionContext, EmitterContext, bUseSwitchState);
-		Emit(Body, *FunctionImplementation);
+		InnerFunctionImplementation(FunctionContext, EmitterContext, bUseSwitchState);
 	}
 
 	const FString ReturnValueString = ReturnValue ? (FString(TEXT(" ")) + FEmitHelper::GetCppName(ReturnValue)) : TEXT("");
-	Emit(Body, *FString::Printf(TEXT("\treturn%s;\n"), *ReturnValueString));
-	Emit(Body, TEXT("}\n\n"));
+	EmitterContext.AddLine(FString::Printf(TEXT("return%s;"), *ReturnValueString));
+	EmitterContext.DecreaseIndent();
+	EmitterContext.AddLine(TEXT("}"));
 }
 
-void FBlueprintCompilerCppBackendBase::GenerateCodeFromEnum(UUserDefinedEnum* SourceEnum)
+FString FBlueprintCompilerCppBackendBase::GenerateCodeFromEnum(UUserDefinedEnum* SourceEnum)
 {
 	check(SourceEnum);
+	FCodeText Header;
+	Header.AddLine(TEXT("#pragma once"));
 	// use GetBaseFilename() so that we can coordinate #includes and filenames
-	EmitFileBeginning(FEmitHelper::GetBaseFilename(SourceEnum), nullptr);
-
-	Emit(Header, TEXT("UENUM(BlueprintType"));
-	EmitReplaceConvertedMetaData(SourceEnum);
-	Emit(Header, TEXT(")\nenum class "));
-	Emit(Header, *FEmitHelper::GetCppName(SourceEnum));
-	Emit(Header, TEXT(" : uint8\n{"));
-
+	Header.AddLine(FString::Printf(TEXT("#include \"%s.generated.h\""), *FEmitHelper::GetBaseFilename(SourceEnum)));
+	Header.AddLine(FString::Printf(TEXT("UENUM(BlueprintType %s )"), *FEmitHelper::ReplaceConvertedMetaData(SourceEnum)));
+	Header.AddLine(FString::Printf(TEXT("enum class %s  : uint8"), *FEmitHelper::GetCppName(SourceEnum)));
+	Header.AddLine(TEXT("{"));
+	Header.IncreaseIndent();
 	for (int32 Index = 0; Index < SourceEnum->NumEnums(); ++Index)
 	{
 		const FString ElemName = SourceEnum->GetEnumName(Index);
@@ -389,44 +374,37 @@ void FBlueprintCompilerCppBackendBase::GenerateCodeFromEnum(UUserDefinedEnum* So
 
 		const FString& DisplayNameMD = SourceEnum->GetMetaData(TEXT("DisplayName"), ElemValue);// TODO: value or index?
 		const FString Meta = DisplayNameMD.IsEmpty() ? FString() : FString::Printf(TEXT("UMETA(DisplayName = \"%s\")"), *DisplayNameMD);
-		Emit(Header, *FString::Printf(TEXT("\n\t%s = %d %s,"), *ElemName, ElemValue, *Meta));
+		Header.AddLine(FString::Printf(TEXT("%s = %d %s,"), *ElemName, ElemValue, *Meta));
 	}
-	Emit(Header, TEXT("\n};\n"));
+	Header.DecreaseIndent();
+	Header.AddLine(TEXT("};"));
+	return Header.Result;
 }
 
-void FBlueprintCompilerCppBackendBase::EmitReplaceConvertedMetaData(UObject* Obj)
-{
-	const FString ReplaceConvertedMD = FEmitHelper::GenerateReplaceConvertedMD(Obj);
-	if (!ReplaceConvertedMD.IsEmpty())
-	{
-		TArray<FString> AdditionalMD;
-		AdditionalMD.Add(ReplaceConvertedMD);
-		Emit(Header, TEXT(","));
-		Emit(Header, *FEmitHelper::HandleMetaData(nullptr, false, &AdditionalMD));
-	}
-}
-
-void FBlueprintCompilerCppBackendBase::GenerateCodeFromStruct(UUserDefinedStruct* SourceStruct)
+FString FBlueprintCompilerCppBackendBase::GenerateCodeFromStruct(UUserDefinedStruct* SourceStruct)
 {
 	check(SourceStruct);
-
 	FGatherConvertedClassDependencies Dependencies(SourceStruct);
 	FEmitterLocalContext EmitterContext(Dependencies);
 	// use GetBaseFilename() so that we can coordinate #includes and filenames
-	EmitFileBeginning(FEmitHelper::GetBaseFilename(SourceStruct), &Dependencies);
+	EmitFileBeginning(FEmitHelper::GetBaseFilename(SourceStruct), EmitterContext);
+
+	EmitterContext.Header.AddLine(FString::Printf(TEXT("USTRUCT(BlueprintType %s)"), *FEmitHelper::ReplaceConvertedMetaData(SourceStruct)));
+	EmitterContext.Header.AddLine(FString::Printf(TEXT("struct %s"), *FEmitHelper::GetCppName(SourceStruct)));
+	EmitterContext.Header.AddLine(TEXT("{"));
+	EmitterContext.Header.AddLine(TEXT("public:"));
+	EmitterContext.Header.IncreaseIndent();
+	EmitterContext.Header.AddLine(TEXT("GENERATED_BODY()"));
+	EmitStructProperties(EmitterContext, SourceStruct);
 
-	const FString NewName = FEmitHelper::GetCppName(SourceStruct);
-	Emit(Header, TEXT("USTRUCT(BlueprintType"));
-	EmitReplaceConvertedMetaData(SourceStruct);
-	Emit(Header, TEXT(")\n"));
+	FEmitDefaultValueHelper::GenerateGetDefaultValue(SourceStruct, EmitterContext);
+	EmitterContext.Header.DecreaseIndent();
+	EmitterContext.Header.AddLine(TEXT("};"));
 
-	Emit(Header, *FString::Printf(TEXT("struct %s\n{\npublic:\n\tGENERATED_BODY()\n"), *NewName));
-	EmitStructProperties(EmitterContext, SourceStruct);
-	Emit(Header, *FEmitDefaultValueHelper::GenerateGetDefaultValue(SourceStruct, Dependencies));
-	Emit(Header, TEXT("};\n"));
+	return EmitterContext.Header.Result;
 }
 
-void FBlueprintCompilerCppBackendBase::GenerateWrapperForClass(UClass* SourceClass)
+FString FBlueprintCompilerCppBackendBase::GenerateWrapperForClass(UClass* SourceClass)
 {
 	FGatherConvertedClassDependencies Dependencies(SourceClass);
 	FEmitterLocalContext EmitterContext(Dependencies);
@@ -454,7 +432,7 @@ void FBlueprintCompilerCppBackendBase::GenerateWrapperForClass(UClass* SourceCla
 	const bool bGenerateAnyFunction = (0 != FunctionsToGenerate.Num()) || (0 != MCDelegateProperties.Num());
 
 	// Include standard stuff
-	EmitFileBeginning(FEmitHelper::GetBaseFilename(SourceClass), &Dependencies, bGenerateAnyFunction, true);
+	EmitFileBeginning(FEmitHelper::GetBaseFilename(SourceClass), EmitterContext, bGenerateAnyFunction, true);
 
 	// DELEGATES
 	const FString DelegatesClassName = FString::Printf(TEXT("U__Delegates__%s"), *FEmitHelper::GetCppName(SourceClass));
@@ -468,20 +446,26 @@ void FBlueprintCompilerCppBackendBase::GenerateWrapperForClass(UClass* SourceCla
 	};
 	if (bGenerateAnyFunction)
 	{
-		Emit(Header, *FString::Printf(TEXT("\nUCLASS()\nclass %s : public UObject\n{\npublic:\n\tGENERATED_BODY()\n"), *DelegatesClassName));
+		EmitterContext.Header.AddLine(TEXT("UCLASS()"));
+		EmitterContext.Header.AddLine(FString::Printf(TEXT("class %s : public UObject"), *DelegatesClassName));
+		EmitterContext.Header.AddLine(TEXT("{"));
+		EmitterContext.Header.AddLine(TEXT("public:"));
+		EmitterContext.Header.IncreaseIndent();
+		EmitterContext.Header.AddLine(TEXT("GENERATED_BODY()"));
 		for (auto Func : FunctionsToGenerate)
 		{
 			FString ParamNumberStr, Parameters;
 			FEmitHelper::ParseDelegateDetails(EmitterContext, Func, Parameters, ParamNumberStr);
-			Emit(Header, *FString::Printf(TEXT("\tDECLARE_DYNAMIC_DELEGATE%s(%s%s);\n"), *ParamNumberStr, *GenerateDelegateTypeName(Func), *Parameters));
+			EmitterContext.Header.AddLine(FString::Printf(TEXT("DECLARE_DYNAMIC_DELEGATE%s(%s%s);"), *ParamNumberStr, *GenerateDelegateTypeName(Func), *Parameters));
 		}
 		for (auto MCDelegateProp : MCDelegateProperties)
 		{
 			FString ParamNumberStr, Parameters;
 			FEmitHelper::ParseDelegateDetails(EmitterContext, MCDelegateProp->SignatureFunction, Parameters, ParamNumberStr);
-			Emit(Header, *FString::Printf(TEXT("\tDECLARE_DYNAMIC_MULTICAST_DELEGATE%s(%s%s);\n"), *ParamNumberStr, *GenerateMulticastDelegateTypeName(MCDelegateProp), *Parameters));
+			EmitterContext.Header.AddLine(FString::Printf(TEXT("DECLARE_DYNAMIC_MULTICAST_DELEGATE%s(%s%s);"), *ParamNumberStr, *GenerateMulticastDelegateTypeName(MCDelegateProp), *Parameters));
 		}
-		Emit(Header, TEXT("\n};\n"));
+		EmitterContext.Header.DecreaseIndent();
+		EmitterContext.Header.AddLine(TEXT("};"));
 	}
 
 	// Begin the struct
@@ -499,9 +483,12 @@ void FBlueprintCompilerCppBackendBase::GenerateWrapperForClass(UClass* SourceCla
 		ParentStruct = FString::Printf(TEXT("FUnconvertedWrapper<%s>"), *FEmitHelper::GetCppName(SourceClass->GetSuperClass()));
 	}
 	
-	Emit(Header, *FString::Printf(TEXT("struct %s : public %s\n{\n"), *WrapperName, *ParentStruct));
+	EmitterContext.Header.AddLine(FString::Printf(TEXT("struct %s : public %s"), *WrapperName, *ParentStruct));
+	EmitterContext.Header.AddLine(TEXT("{"));
+	EmitterContext.Header.IncreaseIndent();
+
 	// Constructor
-	Emit(Header, *FString::Printf(TEXT("\n\t%s(const %s* __InObject) : %s(__InObject){}\n"), *WrapperName, *FEmitHelper::GetCppName(EmitterContext.GetNativeOrConvertedClass(SourceClass)), *ParentStruct));
+	EmitterContext.Header.AddLine(FString::Printf(TEXT("%s(const %s* __InObject) : %s(__InObject){}"), *WrapperName, *FEmitHelper::GetCppName(EmitterContext.GetFirstNativeOrConvertedClass(SourceClass)), *ParentStruct));
 
 	// PROPERTIES:
 	for (auto Property : TFieldRange<UProperty>(SourceClass, EFieldIteratorFlags::ExcludeSuper))
@@ -510,169 +497,177 @@ void FBlueprintCompilerCppBackendBase::GenerateWrapperForClass(UClass* SourceCla
 		const FString TypeDeclaration = Property->IsA<UMulticastDelegateProperty>()
 			? FString::Printf(TEXT("%s::%s"), *DelegatesClassName, *GenerateMulticastDelegateTypeName(CastChecked<UMulticastDelegateProperty>(Property)))
 			: EmitterContext.ExportCppDeclaration(Property, EExportedDeclaration::Parameter, EPropertyExportCPPFlags::CPPF_CustomTypeName | EPropertyExportCPPFlags::CPPF_BlueprintCppBackend | EPropertyExportCPPFlags::CPPF_NoRef, true);
-		Emit(Header, *FString::Printf(TEXT("\t%s& GetRef__%s()"), *TypeDeclaration, *UnicodeToCPPIdentifier(Property->GetName(), false, nullptr)));
-		Emit(Header, TEXT("\n\t{\n\t\t"));
-		Emit(Header, *FString::Printf(TEXT("static const FName __PropertyName(TEXT(\"%s\"));"), *Property->GetName()));
-		Emit(Header, TEXT("\n\t\t"));
-		Emit(Header, *FString::Printf(TEXT("return *(GetClass()->FindPropertyByName(__PropertyName)->ContainerPtrToValuePtr<%s>(__Object));"),*TypeDeclaration));
-		Emit(Header, TEXT("\n\t}\n\n"));
+		EmitterContext.Header.AddLine(FString::Printf(TEXT("%s& GetRef__%s()"), *TypeDeclaration, *UnicodeToCPPIdentifier(Property->GetName(), false, nullptr)));
+		EmitterContext.Header.AddLine(TEXT("{"));
+		EmitterContext.Header.IncreaseIndent();
+		EmitterContext.Header.AddLine(FString::Printf(TEXT("static const FName __PropertyName(TEXT(\"%s\"));"), *Property->GetName()));
+		EmitterContext.Header.AddLine(FString::Printf(TEXT("return *(GetClass()->FindPropertyByName(__PropertyName)->ContainerPtrToValuePtr<%s>(__Object));"), *TypeDeclaration));
+		EmitterContext.Header.DecreaseIndent();
+		EmitterContext.Header.AddLine(TEXT("}"));
 	}
 
 	// FUNCTIONS:
 	for (auto Func : FunctionsToGenerate)
 	{
-		Emit(Header, *FString::Printf(TEXT("\tvoid %s("), *FEmitHelper::GetCppName(Func)));
-		bool bFirst = true;
-		for (TFieldIterator<UProperty> It(Func); It; ++It)
 		{
-			UProperty* Property = *It;
-			if (!Property || !Property->HasAnyPropertyFlags(CPF_Parm))
+			FString DelcareFunction = FString::Printf(TEXT("void %s("), *FEmitHelper::GetCppName(Func));
+			bool bFirst = true;
+			for (TFieldIterator<UProperty> It(Func); It; ++It)
 			{
-				continue;
-			}
+				UProperty* Property = *It;
+				if (!Property || !Property->HasAnyPropertyFlags(CPF_Parm))
+				{
+					continue;
+				}
 
-			if (!bFirst)
-			{
-				Emit(Header, TEXT(", "));
-			}
-			else
-			{
-				bFirst = false;
-			}
+				if (!bFirst)
+				{
+					DelcareFunction += TEXT(", ");
+				}
+				else
+				{
+					bFirst = false;
+				}
 
-			if (Property->HasAnyPropertyFlags(CPF_OutParm))
-			{
-				Emit(Header, TEXT("/*out*/ "));
-			}
+				if (Property->HasAnyPropertyFlags(CPF_OutParm))
+				{
+					DelcareFunction += TEXT("/*out*/ ");
+				}
 
-			const FString TypeDeclaration = EmitterContext.ExportCppDeclaration(Property, EExportedDeclaration::Parameter, EPropertyExportCPPFlags::CPPF_CustomTypeName | EPropertyExportCPPFlags::CPPF_BlueprintCppBackend);
-			Emit(Header, *TypeDeclaration);
+				DelcareFunction += EmitterContext.ExportCppDeclaration(Property, EExportedDeclaration::Parameter, EPropertyExportCPPFlags::CPPF_CustomTypeName | EPropertyExportCPPFlags::CPPF_BlueprintCppBackend);
+			}
+			DelcareFunction += TEXT(")");
+			EmitterContext.Header.AddLine(DelcareFunction);
 		}
-		Emit(Header, TEXT(")\n\t{\n"));
-		Emit(Header, *FString::Printf(TEXT("\t\t%s::%s __Delegate;\n"), *DelegatesClassName, *GenerateDelegateTypeName(Func)));
-		Emit(Header, *FString::Printf(TEXT("\t\tstatic const FName __FunctionName(TEXT(\"%s\"));\n"), *Func->GetName()));
-		Emit(Header, TEXT("\t\t__Delegate.BindUFunction(__Object, __FunctionName);\n"));
-		Emit(Header, TEXT("\t\tcheck(__Delegate.IsBound());\n"));
-		Emit(Header, TEXT("\t\t__Delegate.Execute("));
-		bFirst = true;
-		for (TFieldIterator<UProperty> It(Func); It; ++It)
+
+		EmitterContext.Header.AddLine(TEXT("{"));
+		EmitterContext.Header.IncreaseIndent();
+
+		EmitterContext.Header.AddLine(FString::Printf(TEXT("%s::%s __Delegate;"), *DelegatesClassName, *GenerateDelegateTypeName(Func)));
+		EmitterContext.Header.AddLine(FString::Printf(TEXT("static const FName __FunctionName(TEXT(\"%s\"));"), *Func->GetName()));
+		EmitterContext.Header.AddLine(TEXT("__Delegate.BindUFunction(__Object, __FunctionName);"));
+		EmitterContext.Header.AddLine(TEXT("check(__Delegate.IsBound());"));
+
 		{
-			UProperty* Property = *It;
-			if (!Property || !Property->HasAnyPropertyFlags(CPF_Parm))
+			FString CallFunction(TEXT("__Delegate.Execute("));
+			bool bFirst = true;
+			for (TFieldIterator<UProperty> It(Func); It; ++It)
 			{
-				continue;
-			}
+				UProperty* Property = *It;
+				if (!Property || !Property->HasAnyPropertyFlags(CPF_Parm))
+				{
+					continue;
+				}
 
-			if (!bFirst)
-			{
-				Emit(Header, TEXT(", "));
-			}
-			else
-			{
-				bFirst = false;
-			}
+				if (!bFirst)
+				{
+					CallFunction += TEXT(", ");
+				}
+				else
+				{
+					bFirst = false;
+				}
 
-			Emit(Header, *UnicodeToCPPIdentifier(Property->GetName(), Property->HasAnyPropertyFlags(CPF_Deprecated), TEXT("bpv__")));
+				CallFunction += UnicodeToCPPIdentifier(Property->GetName(), Property->HasAnyPropertyFlags(CPF_Deprecated), TEXT("bpv__"));
+			}
+			CallFunction += TEXT(");");
+			EmitterContext.Header.AddLine(CallFunction);
 		}
-		Emit(Header, TEXT(");\n"));
-		Emit(Header, TEXT("\t}\n"));
+
+		EmitterContext.Header.DecreaseIndent();
+		EmitterContext.Header.AddLine(TEXT("}"));
 	}
 
 	// close struct
-	Emit(Header, TEXT("};\n"));
+	EmitterContext.Header.DecreaseIndent();
+	EmitterContext.Header.AddLine(TEXT("};"));
+
+	return EmitterContext.Header.Result;
 }
 
-void FBlueprintCompilerCppBackendBase::EmitFileBeginning(const FString& CleanName, FGatherConvertedClassDependencies* Dependencies, bool bIncludeGeneratedH, bool bIncludeCodeHelpersInHeader)
+void FBlueprintCompilerCppBackendBase::EmitFileBeginning(const FString& CleanName, FEmitterLocalContext& EmitterContext, bool bIncludeGeneratedH, bool bIncludeCodeHelpersInHeader)
 {
-	Emit(Header, TEXT("#pragma once\n\n"));
+	EmitterContext.Header.AddLine(TEXT("#pragma once"));
 
-	auto EmitIncludeHeader = [&](FStringOutputDevice& Dst, const TCHAR* Message, bool bAddDotH)
+	auto EmitIncludeHeader = [&](FCodeText& Dst, const TCHAR* Message, bool bAddDotH)
 	{
-		Emit(Dst, *FString::Printf(TEXT("#include \"%s%s\"\n"), Message, bAddDotH ? TEXT(".h") : TEXT("")));
+		Dst.AddLine(FString::Printf(TEXT("#include \"%s%s\""), Message, bAddDotH ? TEXT(".h") : TEXT("")));
 	};
-	EmitIncludeHeader(Body, FApp::GetGameName(), true);
-	EmitIncludeHeader(Body, *CleanName, true);
-	if (bIncludeCodeHelpersInHeader)
-	{
-		EmitIncludeHeader(Header, TEXT("GeneratedCodeHelpers"), true);
-	}
-	else
-	{
-		EmitIncludeHeader(Body, TEXT("GeneratedCodeHelpers"), true);
-	}
+	EmitIncludeHeader(EmitterContext.Body, FApp::GetGameName(), true);
+	EmitIncludeHeader(EmitterContext.Body, *CleanName, true);
+	EmitIncludeHeader(bIncludeCodeHelpersInHeader ? EmitterContext.Header : EmitterContext.Body, TEXT("GeneratedCodeHelpers"), true);
 
-	if (Dependencies)
+	FBackendHelperUMG::AdditionalHeaderIncludeForWidget(EmitterContext);
+
+	TSet<FString> AlreadyIncluded;
+	AlreadyIncluded.Add(CleanName);
+	auto EmitInner = [&](FCodeText& Dst, const TSet<UField*>& Src, const TSet<UField*>& Declarations)
 	{
-		Emit(Header, *FBackendHelperUMG::AdditionalHeaderIncludeForWidget(Cast<UClass>(Dependencies->GetOriginalStruct())));
+		auto EngineSourceDir = FPaths::EngineSourceDir();
+		auto GameSourceDir = FPaths::GameSourceDir();
 
-		TSet<FString> AlreadyIncluded;
-		AlreadyIncluded.Add(CleanName);
-		auto EmitInner = [&](FStringOutputDevice& Dst, const TSet<UField*>& Src, const TSet<UField*>& Declarations)
+		for (UField* Field : Src)
 		{
-			auto EngineSourceDir = FPaths::EngineSourceDir();
-			auto GameSourceDir = FPaths::GameSourceDir();
+			check(Field);
+			const bool bWantedType = Field->IsA<UBlueprintGeneratedClass>() || Field->IsA<UUserDefinedEnum>() || Field->IsA<UUserDefinedStruct>();
 
-			for (UField* Field : Src)
+			// Wanted no-native type, that will be converted
+			if (bWantedType)
 			{
-				check(Field);
-				const bool bWantedType = Field->IsA<UBlueprintGeneratedClass>() || Field->IsA<UUserDefinedEnum>() || Field->IsA<UUserDefinedStruct>();
+				// @TODO: Need to query if this asset will actually be converted
 
-				// Wanted no-native type, that will be converted
-				if (bWantedType)
+				const FString Name = Field->GetName();
+				bool bAlreadyIncluded = false;
+				AlreadyIncluded.Add(Name, &bAlreadyIncluded);
+				if (!bAlreadyIncluded)
 				{
-					// @TODO: Need to query if this asset will actually be converted
-
-					const FString Name = Field->GetName();
-					bool bAlreadyIncluded = false;
-					AlreadyIncluded.Add(Name, &bAlreadyIncluded);
-					if (!bAlreadyIncluded)
-					{
-						const FString GeneratedFilename = IBlueprintCompilerCppBackendModule::GetBaseFilename(Field);
-						EmitIncludeHeader(Dst, *GeneratedFilename, true);
-					}
+					const FString GeneratedFilename = IBlueprintCompilerCppBackendModule::GetBaseFilename(Field);
+					EmitIncludeHeader(Dst, *GeneratedFilename, true);
 				}
-				// headers for native items
-				else
+			}
+			// headers for native items
+			else
+			{
+				FString PackPath;
+				if (FSourceCodeNavigation::FindClassHeaderPath(Field, PackPath))
 				{
-					FString PackPath;
-					if (FSourceCodeNavigation::FindClassHeaderPath(Field, PackPath))
+					if (!PackPath.RemoveFromStart(EngineSourceDir))
 					{
-						if (!PackPath.RemoveFromStart(EngineSourceDir))
+						if (!PackPath.RemoveFromStart(GameSourceDir))
 						{
-							if (!PackPath.RemoveFromStart(GameSourceDir))
-							{
-								PackPath = FPaths::GetCleanFilename(PackPath);
-							}
-						}
-						bool bAlreadyIncluded = false;
-						AlreadyIncluded.Add(PackPath, &bAlreadyIncluded);
-						if (!bAlreadyIncluded)
-						{
-							EmitIncludeHeader(Dst, *PackPath, false);
+							PackPath = FPaths::GetCleanFilename(PackPath);
 						}
 					}
+					bool bAlreadyIncluded = false;
+					AlreadyIncluded.Add(PackPath, &bAlreadyIncluded);
+					if (!bAlreadyIncluded)
+					{
+						EmitIncludeHeader(Dst, *PackPath, false);
+					}
 				}
 			}
+		}
 
-			Emit(Dst, TEXT("\n"));
-
-			for (auto Type : Declarations)
+		for (auto Type : Declarations)
+		{
+			if (auto ForwardDeclaredType = Cast<UClass>(Type))
 			{
-				if (auto ForwardDeclaredType = Cast<UClass>(Type))
-				{
-					Emit(Dst, *FString::Printf(TEXT("class %s;\n"), *FEmitHelper::GetCppName(ForwardDeclaredType)));
-				}
+				Dst.AddLine(FString::Printf(TEXT("class %s;"), *FEmitHelper::GetCppName(ForwardDeclaredType)));
 			}
+		}
+	};
 
-			Emit(Dst, TEXT("\n"));
-		};
-
-		EmitInner(Header, Dependencies->IncludeInHeader, Dependencies->DeclareInHeader);
-		EmitInner(Body, Dependencies->IncludeInBody, TSet<UField*>());
-	}
+	EmitInner(EmitterContext.Header, EmitterContext.Dependencies.IncludeInHeader, EmitterContext.Dependencies.DeclareInHeader);
+	EmitInner(EmitterContext.Body, EmitterContext.Dependencies.IncludeInBody, TSet<UField*>());
 
 	if (bIncludeGeneratedH)
 	{
-		Emit(Header, *FString::Printf(TEXT("#include \"%s.generated.h\"\n\n"), *CleanName));
+		EmitterContext.Header.AddLine(FString::Printf(TEXT("#include \"%s.generated.h\""), *CleanName));
 	}
 }
+
+void FBlueprintCompilerCppBackendBase::CleanBackend()
+{
+	StateMapPerFunction.Empty();
+	FunctionIndexMap.Empty();
+}
\ No newline at end of file
diff --git a/Engine/Source/Developer/BlueprintCompilerCppBackend/Private/BlueprintCompilerCppBackendBase.h b/Engine/Source/Developer/BlueprintCompilerCppBackend/Private/BlueprintCompilerCppBackendBase.h
--- a/Engine/Source/Developer/BlueprintCompilerCppBackend/Private/BlueprintCompilerCppBackendBase.h
+++ b/Engine/Source/Developer/BlueprintCompilerCppBackend/Private/BlueprintCompilerCppBackendBase.h
@@ -34,29 +34,16 @@ protected:
 
 	TArray<FFunctionLabelInfo> StateMapPerFunction;
 	TMap<FKismetFunctionContext*, int32> FunctionIndexMap;
-
-	FString CppClassName;
-
-	FStringOutputDevice Header;
-	FStringOutputDevice Body;
 public:
 
 	// IBlueprintCompilerCppBackend implementation
-	virtual void GenerateCodeFromClass(UClass* SourceClass, TIndirectArray<FKismetFunctionContext>& Functions, bool bGenerateStubsOnly) override;
-	virtual void GenerateCodeFromEnum(UUserDefinedEnum* SourceEnum) override;
-	virtual void GenerateCodeFromStruct(UUserDefinedStruct* SourceStruct) override;
-	virtual void GenerateWrapperForClass(UClass* SourceClass) override;
-
-	virtual const FString& GetBody()	const override { return Body; }
-	virtual const FString& GetHeader()	const override { return Header; }
+	virtual FString GenerateCodeFromClass(UClass* SourceClass, TIndirectArray<FKismetFunctionContext>& Functions, bool bGenerateStubsOnly, FString& OutCppBody) override;
+	virtual FString GenerateCodeFromEnum(UUserDefinedEnum* SourceEnum) override;
+	virtual FString GenerateCodeFromStruct(UUserDefinedStruct* SourceStruct) override;
+	virtual FString GenerateWrapperForClass(UClass* SourceClass) override;
 	// end of IBlueprintCompilerCppBackend implementation
 
 	virtual ~FBlueprintCompilerCppBackendBase() {}
-private:
-	void Emit(FStringOutputDevice& Target, const TCHAR* Message)
-	{
-		Target += Message;
-	}
 
 protected:
 	/*
@@ -73,25 +60,26 @@ protected:
 			}
 
 	*/
-	virtual FString InnerFunctionImplementation(FKismetFunctionContext& FunctionContext, FEmitterLocalContext& EmitterContext, bool bUseSwitchState) PURE_VIRTUAL(FBlueprintCompilerCppBackendBase::InnerFunctionImplementation, return FString(););
-
-	void EmitStructProperties(FEmitterLocalContext EmitterContext, UStruct* SourceClass);
-
-	/** Emits local variable declarations for a function */
-	void DeclareLocalVariables(FKismetFunctionContext& FunctionContext, FEmitterLocalContext& EmitterContext, TArray<UProperty*>& LocalVariables);
+	virtual void InnerFunctionImplementation(FKismetFunctionContext& FunctionContext, FEmitterLocalContext& EmitterContext, bool bUseSwitchState) PURE_VIRTUAL(FBlueprintCompilerCppBackendBase::InnerFunctionImplementation,);
 	
-	/** Builds both the header declaration and body implementation of a function */
-	void ConstructFunction(FKismetFunctionContext& FunctionContext, FEmitterLocalContext& EmitterContext, bool bGenerateStubOnly);
-
-	void EmitFileBeginning(const FString& CleanName, FGatherConvertedClassDependencies* Dependencies, bool bIncludeGeneratedH = true, bool bIncludeCodeHelpersInHeader = false);
-
 	int32 StatementToStateIndex(FKismetFunctionContext& FunctionContext, FBlueprintCompiledStatement* Statement)
 	{
 		int32 Index = FunctionIndexMap.FindChecked(&FunctionContext);
 		return StateMapPerFunction[Index].StatementToStateIndex(Statement);
 	}
 
-	void EmitReplaceConvertedMetaData(UObject* Obj);
+private:
+	/** Builds both the header declaration and body implementation of a function */
+	void ConstructFunction(FKismetFunctionContext& FunctionContext, FEmitterLocalContext& EmitterContext, bool bGenerateStubOnly);
+
+	void CleanBackend();
+
+	static void EmitFileBeginning(const FString& CleanName, FEmitterLocalContext& EmitterContext, bool bIncludeGeneratedH = true, bool bIncludeCodeHelpersInHeader = false);
+
+	static void EmitStructProperties(FEmitterLocalContext& EmitterContext, UStruct* SourceClass);
+
+	/** Emits local variable declarations for a function */
+	static void DeclareLocalVariables(FEmitterLocalContext& EmitterContext, TArray<UProperty*>& LocalVariables);
 
-	void DeclareDelegates(FEmitterLocalContext EmitterContext, UClass* SourceClass, TIndirectArray<FKismetFunctionContext>& Functions);
+	static void DeclareDelegates(FEmitterLocalContext& EmitterContext, TIndirectArray<FKismetFunctionContext>& Functions);
 };
diff --git a/Engine/Source/Developer/BlueprintCompilerCppBackend/Private/BlueprintCompilerCppBackendUMG.cpp b/Engine/Source/Developer/BlueprintCompilerCppBackend/Private/BlueprintCompilerCppBackendUMG.cpp
--- a/Engine/Source/Developer/BlueprintCompilerCppBackend/Private/BlueprintCompilerCppBackendUMG.cpp
+++ b/Engine/Source/Developer/BlueprintCompilerCppBackend/Private/BlueprintCompilerCppBackendUMG.cpp
@@ -6,37 +6,37 @@
 #include "WidgetAnimation.h"
 #include "WidgetTree.h"
 
-FString FBackendHelperUMG::WidgetFunctionsInHeader(UClass* SourceClass)
+void FBackendHelperUMG::WidgetFunctionsInHeader(FEmitterLocalContext& Context)
 {
-	if (Cast<UWidgetBlueprintGeneratedClass>(SourceClass))
+	if (Cast<UWidgetBlueprintGeneratedClass>(Context.GetCurrentlyGeneratedClass()))
 	{
-		return FString(TEXT("public:\n\tvirtual void GetSlotNames(TArray<FName>& SlotNames) const override;\n\tvirtual void PreSave() override;\n\tvirtual void CustomNativeInitilize() override;\n"));
+		Context.Header.AddLine(TEXT("virtual void GetSlotNames(TArray<FName>& SlotNames) const override;"));
+		Context.Header.AddLine(TEXT("virtual void PreSave() override;"));
+		Context.Header.AddLine(TEXT("virtual void CustomNativeInitilize() override;"));
 	}
-	return FString();
 }
 
-FString FBackendHelperUMG::AdditionalHeaderIncludeForWidget(UClass* SourceClass)
+void FBackendHelperUMG::AdditionalHeaderIncludeForWidget(FEmitterLocalContext& Context)
 {
-	if (Cast<UWidgetBlueprintGeneratedClass>(SourceClass))
+	if (Cast<UWidgetBlueprintGeneratedClass>(Context.GetCurrentlyGeneratedClass()))
 	{
-		return FString(TEXT("\n#include \"Runtime/UMG/Public/UMG.h\"\n"));
+		Context.Header.AddLine(TEXT("#include \"Runtime/UMG/Public/UMG.h\""));
 	}
-	return FString();
 }
 
-void FBackendHelperUMG::CreateClassSubobjects(FEmitterLocalContext& Context)
+void FBackendHelperUMG::CreateClassSubobjects(FEmitterLocalContext& Context, bool bCreate, bool bInitialize)
 {
 	if (auto WidgetClass = Cast<UWidgetBlueprintGeneratedClass>(Context.GetCurrentlyGeneratedClass()))
 	{
 		if (WidgetClass->WidgetTree)
 		{
 			ensure(WidgetClass->WidgetTree->GetOuter() == Context.GetCurrentlyGeneratedClass());
-			FEmitDefaultValueHelper::HandleClassSubobject(Context, WidgetClass->WidgetTree, FEmitterLocalContext::EClassSubobjectList::MiscConvertedSubobjects);
+			FEmitDefaultValueHelper::HandleClassSubobject(Context, WidgetClass->WidgetTree, FEmitterLocalContext::EClassSubobjectList::MiscConvertedSubobjects, bCreate, bInitialize);
 		}
 		for (auto Anim : WidgetClass->Animations)
 		{
 			ensure(Anim->GetOuter() == Context.GetCurrentlyGeneratedClass());
-			FEmitDefaultValueHelper::HandleClassSubobject(Context, Anim, FEmitterLocalContext::EClassSubobjectList::MiscConvertedSubobjects);
+			FEmitDefaultValueHelper::HandleClassSubobject(Context, Anim, FEmitterLocalContext::EClassSubobjectList::MiscConvertedSubobjects, bCreate, bInitialize);
 		}
 	}
 }
diff --git a/Engine/Source/Developer/BlueprintCompilerCppBackend/Private/BlueprintCompilerCppBackendUtils.cpp b/Engine/Source/Developer/BlueprintCompilerCppBackend/Private/BlueprintCompilerCppBackendUtils.cpp
--- a/Engine/Source/Developer/BlueprintCompilerCppBackend/Private/BlueprintCompilerCppBackendUtils.cpp
+++ b/Engine/Source/Developer/BlueprintCompilerCppBackend/Private/BlueprintCompilerCppBackendUtils.cpp
@@ -4,17 +4,6 @@
 #include "BlueprintCompilerCppBackendUtils.h"
 #include "EdGraphSchema_K2.h"
 
-FString FEmitterLocalContext::GenerateGetProperty(const UProperty* Property) const
-{
-	check(Property);
-	const UStruct* OwnerStruct = Property->GetOwnerStruct();
-	const FString OwnerName = FEmitHelper::GetCppName(OwnerStruct);
-	const FString OwnerPath = OwnerName + (OwnerStruct->IsA<UClass>() ? TEXT("::StaticClass()") : TEXT("::StaticStruct()"));
-	//return FString::Printf(TEXT("%s->FindPropertyByName(GET_MEMBER_NAME_CHECKED(%s, %s))"), *OwnerPath, *OwnerName, *Property->GetNameCPP());
-
-	return FString::Printf(TEXT("%s->FindPropertyByName(FName(TEXT(\"%s\")))"), *OwnerPath, *FEmitHelper::GetCppName(Property));
-}
-
 FString FEmitterLocalContext::GenerateUniqueLocalName()
 {
 	const FString UniqueNameBase = TEXT("__Local__");
@@ -38,7 +27,7 @@ FString FEmitterLocalContext::FindGloballyMappedObject(const UObject* Object, co
 			return *NamePtr;
 		}
 
-		const UClass* ObjectClassToUse = ExpectedClass ? ExpectedClass : GetNativeOrConvertedClass(Object->GetClass());
+		const UClass* ObjectClassToUse = ExpectedClass ? ExpectedClass : GetFirstNativeOrConvertedClass(Object->GetClass());
 		const FString ClassString = FEmitHelper::GetCppName(ObjectClassToUse);
 
 		int32 ObjectsCreatedPerClassIdx = MiscConvertedSubobjects.IndexOfByKey(Object);
@@ -113,7 +102,7 @@ FString FEmitterLocalContext::FindGloballyMappedObject(const UObject* Object, co
 	ensure(!bLoadIfNotFound || Object);
 	if (Object && (bLoadIfNotFound || bTryUsedAssetsList))
 	{
-		const UClass* ObjectClassToUse = ExpectedClass ? ExpectedClass : GetNativeOrConvertedClass(Object->GetClass());
+		const UClass* ObjectClassToUse = ExpectedClass ? ExpectedClass : GetFirstNativeOrConvertedClass(Object->GetClass());
 		const FString ClassString = FEmitHelper::GetCppName(ObjectClassToUse);
 
 		if (bTryUsedAssetsList)
@@ -157,7 +146,7 @@ FString FEmitterLocalContext::ExportTextItem(const UProperty* Property, const vo
 	{
 		if (UObject* Object = ObjectPropertyBase->GetObjectPropertyValue(PropertyValue))
 		{
-			const UClass* ActualClass = GetNativeOrConvertedClass(ObjectPropertyBase->PropertyClass);
+			const UClass* ActualClass = GetFirstNativeOrConvertedClass(ObjectPropertyBase->PropertyClass);
 			return FString::Printf(TEXT("LoadObject<%s>(nullptr, TEXT(\"%s\"))")
 				, *FEmitHelper::GetCppName(ActualClass)
 				, *(Object->GetPathName().ReplaceCharWithEscapedChar()));
@@ -188,7 +177,7 @@ FString FEmitterLocalContext::ExportCppDeclaration(const UProperty* Property, EE
 			{
 				const bool bIsParameter = (DeclarationType == EExportedDeclaration::Parameter) || (DeclarationType == EExportedDeclaration::MacroParameter);
 				const uint32 LocalExportCPPFlags = ExportCPPFlags | (bIsParameter ? CPPF_ArgumentOrReturnValue : 0);
-				UClass* NativeType = GetFirstNativeParent(BPGC);
+				UClass* NativeType = GetFirstNativeOrConvertedClass(BPGC);
 				ActualCppType = ObjectPropertyBase->GetCPPTypeCustom(&ActualExtendedType, LocalExportCPPFlags, NativeType);
 				ActualCppTypePtr = &ActualCppType;
 				if (!ActualExtendedType.IsEmpty())
@@ -638,9 +627,8 @@ int32 FEmitHelper::ParseDelegateDetails(FEmitterLocalContext& EmitterContext, UF
 	return ParameterNum;
 }
 
-TArray<FString> FEmitHelper::EmitSinglecastDelegateDeclarations(FEmitterLocalContext& EmitterContext, const TArray<UDelegateProperty*>& Delegates)
+void FEmitHelper::EmitSinglecastDelegateDeclarations(FEmitterLocalContext& EmitterContext, const TArray<UDelegateProperty*>& Delegates)
 {
-	TArray<FString> Results;
 	for (auto It : Delegates)
 	{
 		check(It);
@@ -656,16 +644,13 @@ TArray<FString> FEmitHelper::EmitSinglecastDelegateDeclarations(FEmitterLocalCon
 			| EPropertyExportCPPFlags::CPPF_NoStaticArray
 			| EPropertyExportCPPFlags::CPPF_BlueprintCppBackend;
 		const FString TypeName = EmitterContext.ExportCppDeclaration(It, EExportedDeclaration::Parameter, LocalExportCPPFlags, true);
-
-		Results.Add(*FString::Printf(TEXT("\tDECLARE_DYNAMIC_DELEGATE%s(%s%s)"), *ParamNumberStr, *TypeName, *Parameters));
+		EmitterContext.Header.AddLine(FString::Printf(TEXT("DECLARE_DYNAMIC_DELEGATE%s(%s%s);"), *ParamNumberStr, *TypeName, *Parameters));
 	}
-	return Results;
 }
 
-TArray<FString> FEmitHelper::EmitMulticastDelegateDeclarations(FEmitterLocalContext& EmitterContext, UClass* SourceClass)
+void FEmitHelper::EmitMulticastDelegateDeclarations(FEmitterLocalContext& EmitterContext)
 {
-	TArray<FString> Results;
-	for (TFieldIterator<UMulticastDelegateProperty> It(SourceClass, EFieldIteratorFlags::ExcludeSuper); It; ++It)
+	for (TFieldIterator<UMulticastDelegateProperty> It(EmitterContext.GetCurrentlyGeneratedClass(), EFieldIteratorFlags::ExcludeSuper); It; ++It)
 	{
 		auto Signature = It->SignatureFunction;
 		check(Signature);
@@ -679,16 +664,14 @@ TArray<FString> FEmitHelper::EmitMulticastDelegateDeclarations(FEmitterLocalCont
 			| EPropertyExportCPPFlags::CPPF_NoStaticArray
 			| EPropertyExportCPPFlags::CPPF_BlueprintCppBackend;
 		const FString TypeName = EmitterContext.ExportCppDeclaration(*It, EExportedDeclaration::Parameter, LocalExportCPPFlags, true);
-
-		Results.Add(*FString::Printf(TEXT("\tDECLARE_DYNAMIC_MULTICAST_DELEGATE%s(%s%s)"), *ParamNumberStr, *TypeName, *Parameters));
+		EmitterContext.Header.AddLine(FString::Printf(TEXT("DECLARE_DYNAMIC_MULTICAST_DELEGATE%s(%s%s);"), *ParamNumberStr, *TypeName, *Parameters));
 	}
-
-	return Results;
 }
 
-FString FEmitHelper::EmitLifetimeReplicatedPropsImpl(UClass* SourceClass, const FString& CppClassName, const TCHAR* InCurrentIndent)
+void FEmitHelper::EmitLifetimeReplicatedPropsImpl(FEmitterLocalContext& EmitterContext)
 {
-	FString Result;
+	auto SourceClass = EmitterContext.GetCurrentlyGeneratedClass();
+	auto CppClassName = FEmitHelper::GetCppName(SourceClass);
 	bool bFunctionInitilzed = false;
 	for (TFieldIterator<UProperty> It(SourceClass, EFieldIteratorFlags::ExcludeSuper); It; ++It)
 	{
@@ -696,19 +679,20 @@ FString FEmitHelper::EmitLifetimeReplicatedPropsImpl(UClass* SourceClass, const
 		{
 			if (!bFunctionInitilzed)
 			{
-				Result += FString::Printf(TEXT("%svoid %s::GetLifetimeReplicatedProps(TArray< FLifetimeProperty > & OutLifetimeProps) const\n"), InCurrentIndent, *CppClassName);
-				Result += FString::Printf(TEXT("%s{\n"), InCurrentIndent);
-				Result += FString::Printf(TEXT("%s\tSuper::GetLifetimeReplicatedProps(OutLifetimeProps);\n"), InCurrentIndent);
+				EmitterContext.AddLine(FString::Printf(TEXT("void %s::GetLifetimeReplicatedProps(TArray< FLifetimeProperty > & OutLifetimeProps) const"), *CppClassName));
+				EmitterContext.AddLine(TEXT("{"));
+				EmitterContext.IncreaseIndent();
+				EmitterContext.AddLine(TEXT("Super::GetLifetimeReplicatedProps(OutLifetimeProps);"));
 				bFunctionInitilzed = true;
 			}
-			Result += FString::Printf(TEXT("%s\tDOREPLIFETIME( %s, %s);\n"), InCurrentIndent, *CppClassName, *FEmitHelper::GetCppName(*It));
+			EmitterContext.AddLine(FString::Printf(TEXT("DOREPLIFETIME(%s, %s);"), *CppClassName, *FEmitHelper::GetCppName(*It)));
 		}
 	}
 	if (bFunctionInitilzed)
 	{
-		Result += FString::Printf(TEXT("%s}\n"), InCurrentIndent);
+		EmitterContext.DecreaseIndent();
+		EmitterContext.AddLine(TEXT("}"));
 	}
-	return Result;
 }
 
 FString FEmitHelper::LiteralTerm(FEmitterLocalContext& EmitterContext, const FEdGraphPinType& Type, const FString& CustomValue, UObject* LiteralObject)
@@ -871,7 +855,7 @@ FString FEmitHelper::LiteralTerm(FEmitterLocalContext& EmitterContext, const FEd
 		if (LiteralObject)
 		{
 			UClass* FoundClass = Cast<UClass>(Type.PinSubCategoryObject.Get());
-			UClass* ObjectClassToUse = FoundClass ? EmitterContext.GetNativeOrConvertedClass(FoundClass) : UObject::StaticClass();
+			UClass* ObjectClassToUse = FoundClass ? EmitterContext.GetFirstNativeOrConvertedClass(FoundClass) : UObject::StaticClass();
 			const FString MappedObject = EmitterContext.FindGloballyMappedObject(LiteralObject, ObjectClassToUse, true);
 			if (!MappedObject.IsEmpty())
 			{
@@ -1003,3 +987,17 @@ bool FEmitHelper::GenerateAutomaticCast(const FEdGraphPinType& LType, const FEdG
 
 	return false;
 }
+
+FString FEmitHelper::ReplaceConvertedMetaData(UObject* Obj)
+{
+	FString Result;
+	const FString ReplaceConvertedMD = FEmitHelper::GenerateReplaceConvertedMD(Obj);
+	if (!ReplaceConvertedMD.IsEmpty())
+	{
+		TArray<FString> AdditionalMD;
+		AdditionalMD.Add(ReplaceConvertedMD);
+		Result += TEXT(",");
+		Result += FEmitHelper::HandleMetaData(nullptr, false, &AdditionalMD);
+	}
+	return Result;
+}
\ No newline at end of file
diff --git a/Engine/Source/Developer/BlueprintCompilerCppBackend/Private/BlueprintCompilerCppBackendUtils.h b/Engine/Source/Developer/BlueprintCompilerCppBackend/Private/BlueprintCompilerCppBackendUtils.h
--- a/Engine/Source/Developer/BlueprintCompilerCppBackend/Private/BlueprintCompilerCppBackendUtils.h
+++ b/Engine/Source/Developer/BlueprintCompilerCppBackend/Private/BlueprintCompilerCppBackendUtils.h
@@ -3,6 +3,27 @@
 
 #include "BlueprintCompilerCppBackendGatherDependencies.h"
 
+struct FCodeText
+{
+	FString Indent;
+	FString Result;
+
+	void IncreaseIndent()
+	{
+		Indent += TEXT("\t");
+	}
+
+	void DecreaseIndent()
+	{
+		Indent.RemoveFromEnd(TEXT("\t"));
+	}
+
+	void AddLine(const FString& Line)
+	{
+		Result += FString::Printf(TEXT("%s%s\n"), *Indent, *Line);
+	}
+};
+
 struct FEmitterLocalContext
 {
 	enum class EClassSubobjectList
@@ -23,18 +44,14 @@ struct FEmitterLocalContext
 	EGeneratedCodeType CurrentCodeType;
 
 private:
-	// TEXT
-	FString Indent;
-	FString Result;
 	int32 LocalNameIndexMax;
-	TArray<FString> HighPriorityLines;
-	TArray<FString> LowPriorityLines;
 
 	//ConstructorOnly Local Names
 	TMap<UObject*, FString> ClassSubobjectsMap;
 	//ConstructorOnly Local Names
 	TMap<UObject*, FString> CommonSubobjectsMap;
 
+	// Class subobjects
 	TArray<UObject*> MiscConvertedSubobjects;
 	TArray<UObject*> DynamicBindingObjects;
 	TArray<UObject*> ComponentTemplates;
@@ -42,18 +59,23 @@ private:
 
 public:
 
+	FCodeText Header;
+	FCodeText Body;
+	FCodeText* DefaultTarget;
+
 	const FGatherConvertedClassDependencies& Dependencies;
 
 	FEmitterLocalContext(const FGatherConvertedClassDependencies& InDependencies)
 		: CurrentCodeType(EGeneratedCodeType::Regular)
 		, LocalNameIndexMax(0)
+		, DefaultTarget(&Body)
 		, Dependencies(InDependencies)
 	{}
 
 	// CONSTRUCTOR FUNCTIONS
 
-	const TCHAR* ClassSubobjectListName(EClassSubobjectList ListType)
-		{
+	static const TCHAR* ClassSubobjectListName(EClassSubobjectList ListType)
+	{
 		if (ListType == EClassSubobjectList::ComponentTemplates)
 		{
 			return TEXT("ComponentTemplates");
@@ -67,49 +89,52 @@ public:
 			return TEXT("DynamicBindingObjects");
 		}
 		return TEXT("MiscConvertedSubobjects");
-		}
+	}
 
 	void RegisterClassSubobject(UObject* Object, EClassSubobjectList ListType)
 	{
+		ensure(CurrentCodeType == FEmitterLocalContext::EGeneratedCodeType::SubobjectsOfClass);
 		if (ListType == EClassSubobjectList::ComponentTemplates)
 		{
 			ComponentTemplates.Add(Object);
-	}
+		}
 		else if (ListType == EClassSubobjectList::Timelines)
-	{
+		{
 			Timelines.Add(Object);
-	}
+		}
 		else if(ListType == EClassSubobjectList::DynamicBindingObjects)
-	{
+		{
 			DynamicBindingObjects.Add(Object);
 		}
 		else
 		{
 			MiscConvertedSubobjects.Add(Object);
 		}
-		}
+	}
 
 	void AddClassSubObject_InConstructor(UObject* Object, const FString& NativeName)
 	{
 		ensure(CurrentCodeType == EGeneratedCodeType::SubobjectsOfClass);
+		ensure(!ClassSubobjectsMap.Contains(Object));
 		ClassSubobjectsMap.Add(Object, NativeName);
 	}
 
 	void AddCommonSubObject_InConstructor(UObject* Object, const FString& NativeName)
 	{
 		ensure(CurrentCodeType == EGeneratedCodeType::CommonConstructor);
-		ClassSubobjectsMap.Add(Object, NativeName);
+		ensure(!CommonSubobjectsMap.Contains(Object));
+		CommonSubobjectsMap.Add(Object, NativeName);
 	}
 
 	// UNIVERSAL FUNCTIONS
 
-	UClass* GetFirstNativeParent(UClass* InClass) const
+	UClass* GetFirstNativeOrConvertedClass(UClass* InClass) const
 	{
 		check(InClass);
 		for (UClass* ItClass = InClass; ItClass; ItClass = ItClass->GetSuperClass())
 		{
 			auto BPGC = Cast<UBlueprintGeneratedClass>(ItClass);
-			if (ItClass->HasAnyClassFlags(CLASS_Native) || (BPGC && Dependencies.WillClassBeConverted(BPGC)))
+			if (ItClass->HasAnyClassFlags(CLASS_Native) || (ensure(BPGC) && Dependencies.WillClassBeConverted(BPGC)))
 			{
 				return ItClass;
 			}
@@ -118,23 +143,6 @@ public:
 		return nullptr;
 	}
 
-	UClass* GetNativeOrConvertedClass(UClass* InClass) const
-	{
-		check(InClass);
-		if (InClass->HasAnyClassFlags(CLASS_Native))
-		{
-			return InClass;
-		}
-		auto BPGC = CastChecked<UBlueprintGeneratedClass>(InClass);
-		if (Dependencies.WillClassBeConverted(BPGC))
-		{
-			return InClass;
-		}
-		return GetFirstNativeParent(InClass);
-	}
-
-	FString GenerateGetProperty(const UProperty* Property) const;
-
 	FString GenerateUniqueLocalName();
 
 	UClass* GetCurrentlyGeneratedClass() const
@@ -146,59 +154,28 @@ public:
 	(due to the native code generation), should be handled by this function */
 	FString FindGloballyMappedObject(const UObject* Object, const UClass* ExpectedClass = nullptr, bool bLoadIfNotFound = false, bool bTryUsedAssetsList = true);
 
-	// TEXT FUNCTIONS
+	// Functions needed for Unconverted classes
+	FString ExportCppDeclaration(const UProperty* Property, EExportedDeclaration::Type DeclarationType, uint32 AdditionalExportCPPFlags, bool bSkipParameterName = false) const;
+	FString ExportTextItem(const UProperty* Property, const void* PropertyValue) const;
 
+	// AS FCodeText
 	void IncreaseIndent()
 	{
-		Indent += TEXT("\t");
+		DefaultTarget->IncreaseIndent();
 	}
 
 	void DecreaseIndent()
 	{
-		Indent.RemoveFromEnd(TEXT("\t"));
-	}
-
-	void AddHighPriorityLine(const FString& Line)
-	{
-		HighPriorityLines.Emplace(FString::Printf(TEXT("%s%s\n"), *Indent, *Line));
+		DefaultTarget->DecreaseIndent();
 	}
 
 	void AddLine(const FString& Line)
 	{
-		LowPriorityLines.Emplace(FString::Printf(TEXT("%s%s\n"), *Indent, *Line));
+		DefaultTarget->AddLine(Line);
 	}
 
-	void FlushLines()
-	{
-		for (auto& Line : HighPriorityLines)
-		{
-			Result += Line;
-		}
-		HighPriorityLines.Reset();
-
-		for (auto& Line : LowPriorityLines)
-		{
-			Result += Line;
-		}
-		LowPriorityLines.Reset();
-	}
-
-	FString GetResult()
-	{
-		FlushLines();
-		return Result;
-	}
-
-	void ClearResult()
-	{
-		FlushLines();
-		Result.Reset();
-	}
-
-	// Functions needed for Unconverted classes
-	FString ExportCppDeclaration(const UProperty* Property, EExportedDeclaration::Type DeclarationType, uint32 AdditionalExportCPPFlags, bool bSkipParameterName = false) const;
-	FString ExportTextItem(const UProperty* Property, const void* PropertyValue) const;
-
+private:
+	FEmitterLocalContext(const FEmitterLocalContext&);
 };
 
 struct FEmitHelper
@@ -228,11 +205,11 @@ struct FEmitHelper
 
 	static int32 ParseDelegateDetails(FEmitterLocalContext& EmitterContext, UFunction* Signature, FString& OutParametersMacro, FString& OutParamNumberStr);
 
-	static TArray<FString> EmitSinglecastDelegateDeclarations(FEmitterLocalContext& EmitterContext, const TArray<UDelegateProperty*>& Delegates);
+	static void EmitSinglecastDelegateDeclarations(FEmitterLocalContext& EmitterContext, const TArray<UDelegateProperty*>& Delegates);
 
-	static TArray<FString> EmitMulticastDelegateDeclarations(FEmitterLocalContext& EmitterContext, UClass* SourceClass);
+	static void EmitMulticastDelegateDeclarations(FEmitterLocalContext& EmitterContext);
 
-	static FString EmitLifetimeReplicatedPropsImpl(UClass* SourceClass, const FString& CppClassName, const TCHAR* InCurrentIndent);
+	static void EmitLifetimeReplicatedPropsImpl(FEmitterLocalContext& EmitterContext);
 
 	static FString LiteralTerm(FEmitterLocalContext& EmitterContext, const FEdGraphPinType& Type, const FString& CustomValue, UObject* LiteralObject);
 
@@ -249,13 +226,17 @@ struct FEmitHelper
 	static FString GenerateReplaceConvertedMD(UObject* Obj);
 
 	static FString GetBaseFilename(const UObject* AssetObj);
+
+	static FString ReplaceConvertedMetaData(UObject* Obj);
 };
 
+struct FNonativeComponentData;
+
 struct FEmitDefaultValueHelper
 {
-	static FString GenerateGetDefaultValue(const UUserDefinedStruct* Struct, const FGatherConvertedClassDependencies& Dependencies);
+	static void GenerateGetDefaultValue(const UUserDefinedStruct* Struct, FEmitterLocalContext& EmitterContext);
 
-	static FString GenerateConstructor(FEmitterLocalContext& Context);
+	static void GenerateConstructor(FEmitterLocalContext& Context);
 
 	enum class EPropertyAccessOperator
 	{
@@ -273,26 +254,27 @@ struct FEmitDefaultValueHelper
 
 	// Creates the subobject (of class) returns it's native local name, 
 	// returns empty string if cannot handle
-	static FString HandleClassSubobject(FEmitterLocalContext& Context, UObject* Object, FEmitterLocalContext::EClassSubobjectList ListOfSubobjectsTyp);
+	static FString HandleClassSubobject(FEmitterLocalContext& Context, UObject* Object, FEmitterLocalContext::EClassSubobjectList ListOfSubobjectsTyp, bool bCreate, bool bInitilize);
 
 private:
 	// Returns native term, 
 	// returns empty string if cannot handle
 	static FString HandleSpecialTypes(FEmitterLocalContext& Context, const UProperty* Property, const uint8* ValuePtr);
 
-	static FString HandleNonNativeComponent(FEmitterLocalContext& Context, const USCS_Node* Node, TSet<const UProperty*>& OutHandledProperties, TArray<FString>& NativeCreatedComponentProperties, const USCS_Node* ParentNode = nullptr);
+	static FString HandleNonNativeComponent(FEmitterLocalContext& Context, const USCS_Node* Node, TSet<const UProperty*>& OutHandledProperties
+		, TArray<FString>& NativeCreatedComponentProperties, const USCS_Node* ParentNode, TArray<FNonativeComponentData>& ComponentsToInit);
 
 	static FString HandleInstancedSubobject(FEmitterLocalContext& Context, UObject* Object, bool bCreateInstance = true, bool bSkipEditorOnlyCheck = false);
 };
 
 struct FBackendHelperUMG
 {
-	static FString WidgetFunctionsInHeader(UClass* SourceClass);
+	static void WidgetFunctionsInHeader(FEmitterLocalContext& Context);
 
-	static FString AdditionalHeaderIncludeForWidget(UClass* SourceClass);
+	static void AdditionalHeaderIncludeForWidget(FEmitterLocalContext& EmitterContext);
 
 	// these function should use the same context as Constructor
-	static void CreateClassSubobjects(FEmitterLocalContext& Context);
+	static void CreateClassSubobjects(FEmitterLocalContext& Context, bool bCreate, bool bInitialize);
 	static void EmitWidgetInitializationFunctions(FEmitterLocalContext& Context);
 
 };
diff --git a/Engine/Source/Developer/BlueprintCompilerCppBackend/Private/BlueprintCompilerCppBackendValueHelper.cpp b/Engine/Source/Developer/BlueprintCompilerCppBackend/Private/BlueprintCompilerCppBackendValueHelper.cpp
--- a/Engine/Source/Developer/BlueprintCompilerCppBackend/Private/BlueprintCompilerCppBackendValueHelper.cpp
+++ b/Engine/Source/Developer/BlueprintCompilerCppBackend/Private/BlueprintCompilerCppBackendValueHelper.cpp
@@ -59,7 +59,15 @@ void FEmitDefaultValueHelper::OuterGenerate(FEmitterLocalContext& Context
 			{
 				ensure(EPropertyAccessOperator::None != AccessOperator);
 				
-				const FString PropertyObject = Context.GenerateGetProperty(Property); //X::StaticClass()->FindPropertyByName(GET_MEMBER_NAME_CHECKED(X, X))
+				auto GenerateGetProperty = [](const UProperty* InProperty) -> FString
+				{
+					check(InProperty);
+					const UStruct* OwnerStruct = InProperty->GetOwnerStruct();
+					const FString OwnerName = FEmitHelper::GetCppName(OwnerStruct);
+					const FString OwnerPath = OwnerName + (OwnerStruct->IsA<UClass>() ? TEXT("::StaticClass()") : TEXT("::StaticStruct()"));
+					return FString::Printf(TEXT("%s->FindPropertyByName(FName(TEXT(\"%s\")))"), *OwnerPath, *FEmitHelper::GetCppName(InProperty));
+				};
+				const FString PropertyObject = GenerateGetProperty(Property); //X::StaticClass()->FindPropertyByName(GET_MEMBER_NAME_CHECKED(X, X))
 
 				const uint32 CppTemplateTypeFlags = EPropertyExportCPPFlags::CPPF_CustomTypeName
 					| EPropertyExportCPPFlags::CPPF_NoConst | EPropertyExportCPPFlags::CPPF_NoRef | EPropertyExportCPPFlags::CPPF_NoStaticArray
@@ -87,25 +95,30 @@ void FEmitDefaultValueHelper::OuterGenerate(FEmitterLocalContext& Context
 	}
 }
 
-FString FEmitDefaultValueHelper::GenerateGetDefaultValue(const UUserDefinedStruct* Struct, const FGatherConvertedClassDependencies& Dependencies)
+void FEmitDefaultValueHelper::GenerateGetDefaultValue(const UUserDefinedStruct* Struct, FEmitterLocalContext& Context)
 {
 	check(Struct);
 	const FString StructName = FEmitHelper::GetCppName(Struct);
-	FString Result = FString::Printf(TEXT("\tstatic %s GetDefaultValue()\n\t{\n\t\t%s DefaultData__;\n"), *StructName, *StructName);
-
-	FStructOnScope StructData(Struct);
-	FStructureEditorUtils::Fill_MakeStructureDefaultValue(Struct, StructData.GetStructMemory());
+	Context.Header.AddLine(FString::Printf(TEXT("static %s GetDefaultValue()"), *StructName));
+	Context.Header.AddLine(TEXT("{"));
 
-	FEmitterLocalContext Context(Dependencies);
-	Context.IncreaseIndent();
-	Context.IncreaseIndent();
-	for (auto Property : TFieldRange<const UProperty>(Struct))
+	Context.Header.IncreaseIndent();
+	Context.Header.AddLine(FString::Printf(TEXT("%s DefaultData__;"), *StructName));
 	{
-		OuterGenerate(Context, Property, TEXT("DefaultData__"), StructData.GetStructMemory(), nullptr, EPropertyAccessOperator::Dot);
+		TGuardValue<FCodeText*> OriginalDefaultTarget(Context.DefaultTarget, &Context.Header);
+		FStructOnScope StructData(Struct);
+		FStructureEditorUtils::Fill_MakeStructureDefaultValue(Struct, StructData.GetStructMemory());
+		for (auto Property : TFieldRange<const UProperty>(Struct))
+		{
+			OuterGenerate(Context, Property, TEXT("DefaultData__"), StructData.GetStructMemory(), nullptr, EPropertyAccessOperator::Dot);
+		}
+
+
 	}
-	Result += Context.GetResult();
-	Result += TEXT("\n\t\treturn DefaultData__;\n\t}\n");
-	return Result;
+	Context.Header.AddLine(TEXT("return DefaultData__;"));
+	Context.Header.DecreaseIndent();
+
+	Context.Header.AddLine(TEXT("}"));
 }
 
 void FEmitDefaultValueHelper::InnerGenerate(FEmitterLocalContext& Context, const UProperty* Property, const FString& PathToMember, const uint8* ValuePtr, const uint8* DefaultValuePtr, bool bWithoutFirstConstructionLine)
@@ -192,7 +205,7 @@ FString FEmitDefaultValueHelper::HandleSpecialTypes(FEmitterLocalContext& Contex
 		if (Object)
 		{
 			{
-				UClass* ObjectClassToUse = Context.GetNativeOrConvertedClass(ObjectProperty->PropertyClass);
+				UClass* ObjectClassToUse = Context.GetFirstNativeOrConvertedClass(ObjectProperty->PropertyClass);
 				const FString MappedObject = Context.FindGloballyMappedObject(Object, ObjectClassToUse);
 				if (!MappedObject.IsEmpty())
 				{
@@ -206,7 +219,7 @@ FString FEmitDefaultValueHelper::HandleSpecialTypes(FEmitterLocalContext& Contex
 				auto CDO = BPGC ? BPGC->GetDefaultObject(false) : nullptr;
 				if (BPGC && Object && CDO && Object->IsIn(BPGC) && !Object->IsIn(CDO) && bCreatingSubObjectsOfClass)
 				{
-					return HandleClassSubobject(Context, Object, FEmitterLocalContext::EClassSubobjectList::MiscConvertedSubobjects);
+					return HandleClassSubobject(Context, Object, FEmitterLocalContext::EClassSubobjectList::MiscConvertedSubobjects, true, true);
 				}
 			}
 
@@ -255,7 +268,48 @@ FString FEmitDefaultValueHelper::HandleSpecialTypes(FEmitterLocalContext& Contex
 	return FString();
 }
 
-FString FEmitDefaultValueHelper::HandleNonNativeComponent(FEmitterLocalContext& Context, const USCS_Node* Node, TSet<const UProperty*>& OutHandledProperties, TArray<FString>& NativeCreatedComponentProperties, const USCS_Node* ParentNode)
+struct FNonativeComponentData
+{
+	FString NativeVariablePropertyName;
+	UActorComponent* ComponentTemplate;
+	UObject* ObjectToCompare;
+
+	////
+	FString ParentVariableName;
+	bool bSetNativeCreationMethod;
+
+	FNonativeComponentData()
+		: ComponentTemplate(nullptr)
+		, ObjectToCompare(nullptr)
+		, bSetNativeCreationMethod(false)
+	{
+	}
+
+	void EmitProperties(FEmitterLocalContext& Context)
+	{
+		ensure(!NativeVariablePropertyName.IsEmpty());
+		if (bSetNativeCreationMethod)
+		{
+			Context.AddLine(FString::Printf(TEXT("%s->CreationMethod = EComponentCreationMethod::Native;"), *NativeVariablePropertyName));
+		}
+
+		if (!ParentVariableName.IsEmpty())
+		{
+			Context.AddLine(FString::Printf(TEXT("%s->AttachParent = %s;"), *NativeVariablePropertyName, *ParentVariableName));
+		}
+
+		UClass* ComponentClass = ComponentTemplate->GetClass();
+		for (auto Property : TFieldRange<const UProperty>(ComponentClass))
+		{
+			FEmitDefaultValueHelper::OuterGenerate(Context, Property, NativeVariablePropertyName
+				, reinterpret_cast<const uint8*>(ComponentTemplate)
+				, reinterpret_cast<const uint8*>(ObjectToCompare)
+				, FEmitDefaultValueHelper::EPropertyAccessOperator::Pointer);
+		}
+	}
+};
+
+FString FEmitDefaultValueHelper::HandleNonNativeComponent(FEmitterLocalContext& Context, const USCS_Node* Node, TSet<const UProperty*>& OutHandledProperties, TArray<FString>& NativeCreatedComponentProperties, const USCS_Node* ParentNode, TArray<FNonativeComponentData>& ComponenntsToInit)
 {
 	check(Node);
 	check(Context.CurrentCodeType == FEmitterLocalContext::EGeneratedCodeType::CommonConstructor);
@@ -281,6 +335,10 @@ FString FEmitDefaultValueHelper::HandleNonNativeComponent(FEmitterLocalContext&
 
 		if (ComponentTemplate->GetOuter() == BPGC)
 		{
+			FNonativeComponentData NonativeComponentData;
+			NonativeComponentData.NativeVariablePropertyName = NativeVariablePropertyName;
+			NonativeComponentData.ComponentTemplate = ComponentTemplate;
+
 			UClass* ComponentClass = ComponentTemplate->GetClass();
 			check(ComponentClass != nullptr);
 
@@ -292,13 +350,13 @@ FString FEmitDefaultValueHelper::HandleNonNativeComponent(FEmitterLocalContext&
 			}
 			else
 			{
-				Context.AddHighPriorityLine(FString::Printf(TEXT("%s%s = CreateDefaultSubobject<%s>(TEXT(\"%s\"));")
+				Context.AddLine(FString::Printf(TEXT("%s%s = CreateDefaultSubobject<%s>(TEXT(\"%s\"));")
 					, (VariableProperty == nullptr) ? TEXT("auto ") : TEXT("")
 					, *NativeVariablePropertyName
 					, *FEmitHelper::GetCppName(ComponentClass)
 					, *VariableCleanName));
 
-				Context.AddLine(FString::Printf(TEXT("%s->CreationMethod = EComponentCreationMethod::Native;"), *NativeVariablePropertyName));
+				NonativeComponentData.bSetNativeCreationMethod = true;
 				NativeCreatedComponentProperties.Add(NativeVariablePropertyName);
 
 				FString ParentVariableName;
@@ -312,27 +370,17 @@ FString FEmitDefaultValueHelper::HandleNonNativeComponent(FEmitterLocalContext&
 				{
 					ParentVariableName = Context.FindGloballyMappedObject(ParentComponentTemplate, USceneComponent::StaticClass());
 				}
-
-				if (!ParentVariableName.IsEmpty())
-				{
-					Context.AddLine(FString::Printf(TEXT("%s->AttachParent = %s;"), *NativeVariablePropertyName, *ParentVariableName));
-				}
-			}
-
-			for (auto Property : TFieldRange<const UProperty>(ComponentClass))
-			{
-				OuterGenerate(Context, Property, NativeVariablePropertyName
-					, reinterpret_cast<const uint8*>(ComponentTemplate)
-					, reinterpret_cast<const uint8*>(ObjectToCompare)
-					, EPropertyAccessOperator::Pointer);
+				NonativeComponentData.ParentVariableName = ParentVariableName;
 			}
+			NonativeComponentData.ObjectToCompare = ObjectToCompare;
+			ComponenntsToInit.Add(NonativeComponentData);
 		}
 	}
 
 	// Recursively handle child nodes.
 	for (auto ChildNode : Node->ChildNodes)
 	{
-		HandleNonNativeComponent(Context, ChildNode, OutHandledProperties, NativeCreatedComponentProperties, Node);
+		HandleNonNativeComponent(Context, ChildNode, OutHandledProperties, NativeCreatedComponentProperties, Node, ComponenntsToInit);
 	}
 
 	return NativeVariablePropertyName;
@@ -448,7 +496,7 @@ struct FDependenciesHelper
 	}
 };
 
-FString FEmitDefaultValueHelper::GenerateConstructor(FEmitterLocalContext& Context)
+void FEmitDefaultValueHelper::GenerateConstructor(FEmitterLocalContext& Context)
 {
 	auto BPGC = CastChecked<UBlueprintGeneratedClass>(Context.GetCurrentlyGeneratedClass());
 	const FString CppClassName = FEmitHelper::GetCppName(BPGC);
@@ -488,32 +536,36 @@ FString FEmitDefaultValueHelper::GenerateConstructor(FEmitterLocalContext& Conte
 		Context.AddLine(TEXT("ensure(0 == CastChecked<UDynamicClass>(GetClass())->ComponentTemplates.Num());"));
 		Context.AddLine(TEXT("ensure(0 == CastChecked<UDynamicClass>(GetClass())->Timelines.Num());"));
 		Context.CurrentCodeType = FEmitterLocalContext::EGeneratedCodeType::SubobjectsOfClass;
-		Context.FlushLines();
-		for (auto ComponentTemplate : ActorComponentTempatesOwnedByClass)
+
+		auto CreateAndInitializeClassSubobjects = [&](bool bCreate, bool bInitilize)
 		{
-			if (ComponentTemplate)
+			for (auto ComponentTemplate : ActorComponentTempatesOwnedByClass)
 			{
-				HandleClassSubobject(Context, ComponentTemplate, FEmitterLocalContext::EClassSubobjectList::ComponentTemplates);
+				if (ComponentTemplate)
+				{
+					HandleClassSubobject(Context, ComponentTemplate, FEmitterLocalContext::EClassSubobjectList::ComponentTemplates, bCreate, bInitilize);
+				}
 			}
-		}
 
-		for (auto TimelineTemplate : BPGC->Timelines)
-		{
-			if (TimelineTemplate)
+			for (auto TimelineTemplate : BPGC->Timelines)
 			{
-				HandleClassSubobject(Context, TimelineTemplate, FEmitterLocalContext::EClassSubobjectList::Timelines);
+				if (TimelineTemplate)
+				{
+					HandleClassSubobject(Context, TimelineTemplate, FEmitterLocalContext::EClassSubobjectList::Timelines, bCreate, bInitilize);
+				}
 			}
-		}
 
-		for (auto DynamicBindingObject : BPGC->DynamicBindingObjects)
-		{
-			if (DynamicBindingObject)
+			for (auto DynamicBindingObject : BPGC->DynamicBindingObjects)
 			{
-				HandleClassSubobject(Context, DynamicBindingObject, FEmitterLocalContext::EClassSubobjectList::DynamicBindingObjects);
+				if (DynamicBindingObject)
+				{
+					HandleClassSubobject(Context, DynamicBindingObject, FEmitterLocalContext::EClassSubobjectList::DynamicBindingObjects, bCreate, bInitilize);
+				}
 			}
-		}
-
-		FBackendHelperUMG::CreateClassSubobjects(Context);
+			FBackendHelperUMG::CreateClassSubobjects(Context, bCreate, bInitilize);
+		};
+		CreateAndInitializeClassSubobjects(true, false);
+		CreateAndInitializeClassSubobjects(false, true);
 		FDependenciesHelper::AddDependenciesInConstructor(Context);
 
 		Context.DecreaseIndent();
@@ -528,7 +580,6 @@ FString FEmitDefaultValueHelper::GenerateConstructor(FEmitterLocalContext& Conte
 		// Let's have an easy access to generated class subobjects
 		Context.AddLine(TEXT("{")); // no shadow variables
 		Context.IncreaseIndent();
-		Context.FlushLines();
 
 		UObject* CDO = BPGC->GetDefaultObject(false);
 
@@ -598,6 +649,8 @@ FString FEmitDefaultValueHelper::GenerateConstructor(FEmitterLocalContext& Conte
 		const bool bErrorFree = UBlueprintGeneratedClass::GetGeneratedClassesHierarchy(BPGC, BPGCStack);
 		if (bErrorFree)
 		{
+			TArray<FNonativeComponentData> ComponentsToInit;
+
 			// Start at the base of the hierarchy so that dependencies are handled first.
 			for (int32 i = BPGCStack.Num() - 1; i >= 0; --i)
 			{
@@ -607,7 +660,7 @@ FString FEmitDefaultValueHelper::GenerateConstructor(FEmitterLocalContext& Conte
 					{
 						if (Node)
 						{
-							const FString NativeVariablePropertyName = HandleNonNativeComponent(Context, Node, HandledProperties, NativeCreatedComponentProperties);
+							const FString NativeVariablePropertyName = HandleNonNativeComponent(Context, Node, HandledProperties, NativeCreatedComponentProperties, nullptr, ComponentsToInit);
 
 							if (i == 0 && bNeedsRootComponentAssignment && Node->ComponentTemplate && Node->ComponentTemplate->IsA<USceneComponent>() && !NativeVariablePropertyName.IsEmpty())
 							{
@@ -620,6 +673,11 @@ FString FEmitDefaultValueHelper::GenerateConstructor(FEmitterLocalContext& Conte
 					}
 				}
 			}
+
+			for (auto& ComponentToInit : ComponentsToInit)
+			{
+				ComponentToInit.EmitProperties(Context);
+			}
 		}
 
 		// Generate ctor init code for generated Blueprint class property values that may differ from parent class defaults (or that otherwise belong to the generated Blueprint class).
@@ -662,48 +720,57 @@ FString FEmitDefaultValueHelper::GenerateConstructor(FEmitterLocalContext& Conte
 	FDependenciesHelper::AddStaticFunctionsForDependencies(Context);
 
 	FBackendHelperUMG::EmitWidgetInitializationFunctions(Context);
-
-	const FString Result = Context.GetResult();
-	Context.ClearResult();
-	return Result;
 }
 
-FString FEmitDefaultValueHelper::HandleClassSubobject(FEmitterLocalContext& Context, UObject* Object, FEmitterLocalContext::EClassSubobjectList ListOfSubobjectsType)
+FString FEmitDefaultValueHelper::HandleClassSubobject(FEmitterLocalContext& Context, UObject* Object, FEmitterLocalContext::EClassSubobjectList ListOfSubobjectsType, bool bCreate, bool bInitilize)
 {
 	ensure(Context.CurrentCodeType == FEmitterLocalContext::EGeneratedCodeType::SubobjectsOfClass);
-	const FString OuterStr = Context.FindGloballyMappedObject(Object->GetOuter());
-	if (OuterStr.IsEmpty())
-	{
-		ensure(false);
-		return FString();
-	}
 
-	const bool AddAsSubobjectOfClass = Object->GetOuter() == Context.GetCurrentlyGeneratedClass();
-	const FString LocalNativeName = Context.GenerateUniqueLocalName();
-	Context.AddClassSubObject_InConstructor(Object, LocalNativeName);
-	UClass* ObjectClass = Object->GetClass();
-	Context.AddHighPriorityLine(FString::Printf(
-		TEXT("auto %s = NewObject<%s>(%s, TEXT(\"%s\"));")
-		, *LocalNativeName
-		, *FEmitHelper::GetCppName(ObjectClass)
-		, *OuterStr
-		, *Object->GetName()));
-	if (AddAsSubobjectOfClass)
+	FString LocalNativeName;
+	if (bCreate)
 	{
-		Context.RegisterClassSubobject(Object, ListOfSubobjectsType);
-		Context.AddHighPriorityLine(FString::Printf(TEXT("CastChecked<UDynamicClass>(GetClass())->%s.Add(%s);")
-			, Context.ClassSubobjectListName(ListOfSubobjectsType)
-			, *LocalNativeName));
+		const FString OuterStr = Context.FindGloballyMappedObject(Object->GetOuter());
+		if (OuterStr.IsEmpty())
+		{
+			ensure(false);
+			return FString();
+		}
+
+		const bool AddAsSubobjectOfClass = Object->GetOuter() == Context.GetCurrentlyGeneratedClass();
+		LocalNativeName = Context.GenerateUniqueLocalName();
+		Context.AddClassSubObject_InConstructor(Object, LocalNativeName);
+		UClass* ObjectClass = Object->GetClass();
+		Context.AddLine(FString::Printf(
+			TEXT("auto %s = NewObject<%s>(%s, TEXT(\"%s\"));")
+			, *LocalNativeName
+			, *FEmitHelper::GetCppName(ObjectClass)
+			, *OuterStr
+			, *Object->GetName()));
+		if (AddAsSubobjectOfClass)
+		{
+			Context.RegisterClassSubobject(Object, ListOfSubobjectsType);
+			Context.AddLine(FString::Printf(TEXT("CastChecked<UDynamicClass>(GetClass())->%s.Add(%s);")
+				, Context.ClassSubobjectListName(ListOfSubobjectsType)
+				, *LocalNativeName));
+		}
 	}
 
-	for (auto Property : TFieldRange<const UProperty>(ObjectClass))
+	if (bInitilize)
 	{
-		OuterGenerate(Context, Property, LocalNativeName
-			, reinterpret_cast<const uint8*>(Object)
-			, reinterpret_cast<const uint8*>(ObjectClass->GetDefaultObject(false))
-			, EPropertyAccessOperator::Pointer);
+		if (LocalNativeName.IsEmpty())
+		{
+			LocalNativeName = Context.FindGloballyMappedObject(Object);
+		}
+		ensure(!LocalNativeName.IsEmpty());
+		auto CDO = Object->GetClass()->GetDefaultObject(false);
+		for (auto Property : TFieldRange<const UProperty>(Object->GetClass()))
+		{
+			OuterGenerate(Context, Property, LocalNativeName
+				, reinterpret_cast<const uint8*>(Object)
+				, reinterpret_cast<const uint8*>(CDO)
+				, EPropertyAccessOperator::Pointer);
+		}
 	}
-
 	return LocalNativeName;
 }
 
@@ -754,13 +821,13 @@ FString FEmitDefaultValueHelper::HandleInstancedSubobject(FEmitterLocalContext&
 	{
 		if (bCreateInstance)
 		{
-			Context.AddHighPriorityLine(FString::Printf(TEXT("auto %s = CreateDefaultSubobject<%s>(TEXT(\"%s\"));")
+			Context.AddLine(FString::Printf(TEXT("auto %s = CreateDefaultSubobject<%s>(TEXT(\"%s\"));")
 				, *LocalNativeName, *FEmitHelper::GetCppName(ObjectClass), *Object->GetName()));
 		}
 		else
 		{
-		Context.AddHighPriorityLine(FString::Printf(TEXT("auto %s = CastChecked<%s>(GetDefaultSubobjectByName(TEXT(\"%s\")));")
-			, *LocalNativeName, *FEmitHelper::GetCppName(ObjectClass), *Object->GetName()));
+			Context.AddLine(FString::Printf(TEXT("auto %s = CastChecked<%s>(GetDefaultSubobjectByName(TEXT(\"%s\")));")
+				, *LocalNativeName, *FEmitHelper::GetCppName(ObjectClass), *Object->GetName()));
 		}
 
 		const UObject* ObjectArchetype = Object->GetArchetype();
@@ -780,7 +847,7 @@ FString FEmitDefaultValueHelper::HandleInstancedSubobject(FEmitterLocalContext&
 			ensure(false);
 			return FString();
 		}
-		Context.AddHighPriorityLine(FString::Printf(TEXT("auto %s = NewObject<%s>(%s, TEXT(\"%s\"));")
+		Context.AddLine(FString::Printf(TEXT("auto %s = NewObject<%s>(%s, TEXT(\"%s\"));")
 			, *LocalNativeName
 			, *FEmitHelper::GetCppName(ObjectClass)
 			, *OuterStr
diff --git a/Engine/Source/Editor/KismetCompiler/Private/KismetCompiler.cpp b/Engine/Source/Editor/KismetCompiler/Private/KismetCompiler.cpp
--- a/Engine/Source/Editor/KismetCompiler/Private/KismetCompiler.cpp
+++ b/Engine/Source/Editor/KismetCompiler/Private/KismetCompiler.cpp
@@ -3409,25 +3409,28 @@ void FKismetCompilerContext::Compile()
 		// Generate code thru the backend(s)
 		if ((bDisplayCpp && bIsFullCompile) || CompileOptions.DoesRequireCppCodeGeneration())
 		{
-			TUniquePtr<IBlueprintCompilerCppBackend> Backend_CPP(IBlueprintCompilerCppBackendModuleInterface::Get().Create());
+			FString CppSourceCode;
+			FString HeaderSourceCode;
 
-			// The C++ backend is currently only for debugging, so it's only run if the output will be visible
-			Backend_CPP->GenerateCodeFromClass(NewClass, FunctionList, !bIsFullCompile);
+			{
+				TUniquePtr<IBlueprintCompilerCppBackend> Backend_CPP(IBlueprintCompilerCppBackendModuleInterface::Get().Create());
+				HeaderSourceCode = Backend_CPP->GenerateCodeFromClass(NewClass, FunctionList, !bIsFullCompile, CppSourceCode);
+			}
 
 			if (CompileOptions.OutHeaderSourceCode.IsValid())
 			{
-				*CompileOptions.OutHeaderSourceCode = Backend_CPP->GetHeader();
+				*CompileOptions.OutHeaderSourceCode = HeaderSourceCode;
 			}
 
 			if (CompileOptions.OutCppSourceCode.IsValid())
 			{
-				*CompileOptions.OutCppSourceCode = Backend_CPP->GetBody();
+				*CompileOptions.OutCppSourceCode = CppSourceCode;
 			}
 
 			if (bDisplayCpp)
 			{
-				UE_LOG(LogK2Compiler, Log, TEXT("[header]\n\n\n%s"), *Backend_CPP->GetHeader());
-				UE_LOG(LogK2Compiler, Log, TEXT("[body]\n\n\n%s"), *Backend_CPP->GetBody());
+				UE_LOG(LogK2Compiler, Log, TEXT("[header]\n\n\n%s"), *HeaderSourceCode);
+				UE_LOG(LogK2Compiler, Log, TEXT("[body]\n\n\n%s"), *CppSourceCode);
 			}
 		}
 
diff --git a/Engine/Source/Editor/KismetCompiler/Private/KismetCompilerModule.cpp b/Engine/Source/Editor/KismetCompiler/Private/KismetCompilerModule.cpp
--- a/Engine/Source/Editor/KismetCompiler/Private/KismetCompilerModule.cpp
+++ b/Engine/Source/Editor/KismetCompiler/Private/KismetCompilerModule.cpp
@@ -132,22 +132,19 @@ void FKismet2CompilerModule::CompileStructure(UUserDefinedStruct* Struct, FCompi
 FString FKismet2CompilerModule::GenerateCppCodeForEnum(UUserDefinedEnum* UDEnum)
 {
 	TUniquePtr<IBlueprintCompilerCppBackend> Backend_CPP(IBlueprintCompilerCppBackendModuleInterface::Get().Create());
-	Backend_CPP->GenerateCodeFromEnum(UDEnum);
-	return Backend_CPP->GetHeader();
+	return Backend_CPP->GenerateCodeFromEnum(UDEnum);
 }
 
 FString FKismet2CompilerModule::GenerateCppCodeForStruct(UUserDefinedStruct* UDStruct)
 {
 	TUniquePtr<IBlueprintCompilerCppBackend> Backend_CPP(IBlueprintCompilerCppBackendModuleInterface::Get().Create());
-	Backend_CPP->GenerateCodeFromStruct(UDStruct);
-	return Backend_CPP->GetHeader();
+	return Backend_CPP->GenerateCodeFromStruct(UDStruct);
 }
 
 FString FKismet2CompilerModule::GenerateCppWrapper(UBlueprintGeneratedClass* BPGC)
 {
 	TUniquePtr<IBlueprintCompilerCppBackend> Backend_CPP(IBlueprintCompilerCppBackendModuleInterface::Get().Create());
-	Backend_CPP->GenerateWrapperForClass(BPGC);
-	return Backend_CPP->GetHeader();
+	return Backend_CPP->GenerateWrapperForClass(BPGC);
 }
 
 extern UNREALED_API FSecondsCounterData BlueprintCompileAndLoadTimerData;
diff --git a/Engine/Source/Editor/KismetCompiler/Public/BlueprintCompilerCppBackendInterface.h b/Engine/Source/Editor/KismetCompiler/Public/BlueprintCompilerCppBackendInterface.h
--- a/Engine/Source/Editor/KismetCompiler/Public/BlueprintCompilerCppBackendInterface.h
+++ b/Engine/Source/Editor/KismetCompiler/Public/BlueprintCompilerCppBackendInterface.h
@@ -8,15 +8,12 @@
 class IBlueprintCompilerCppBackend
 {
 public:
-	virtual void GenerateCodeFromClass(UClass* SourceClass, TIndirectArray<FKismetFunctionContext>& Functions, bool bGenerateStubsOnly) = 0;
-	virtual void GenerateCodeFromEnum(UUserDefinedEnum* SourceEnum) = 0;
-	virtual void GenerateCodeFromStruct(UUserDefinedStruct* SourceStruct) = 0;
+	virtual FString GenerateCodeFromClass(UClass* SourceClass, TIndirectArray<FKismetFunctionContext>& Functions, bool bGenerateStubsOnly, FString& OutCppBody) = 0;
+	virtual FString GenerateCodeFromEnum(UUserDefinedEnum* SourceEnum) = 0;
+	virtual FString GenerateCodeFromStruct(UUserDefinedStruct* SourceStruct) = 0;
 
 	// Generate a wrapper class, that helps accessing non-native properties and calling non-native functions
-	virtual void GenerateWrapperForClass(UClass* SourceClass) = 0;
-
-	virtual const FString& GetBody() const = 0;
-	virtual const FString& GetHeader() const = 0;
+	virtual FString GenerateWrapperForClass(UClass* SourceClass) = 0;
 
 	virtual ~IBlueprintCompilerCppBackend() {}
 };
