FEmitHelper::GenerateAutomaticCast detects unconverted classes.

diff --git a/Engine/Source/Developer/BlueprintCompilerCppBackend/Private/BlueprintCompilerCppBackend.cpp b/Engine/Source/Developer/BlueprintCompilerCppBackend/Private/BlueprintCompilerCppBackend.cpp
--- a/Engine/Source/Developer/BlueprintCompilerCppBackend/Private/BlueprintCompilerCppBackend.cpp
+++ b/Engine/Source/Developer/BlueprintCompilerCppBackend/Private/BlueprintCompilerCppBackend.cpp
@@ -99,7 +99,7 @@ void FBlueprintCompilerCppBackend::EmitAssignmentStatment(FEmitterLocalContext&
 
 	FString BeginCast;
 	FString EndCast;
-	FEmitHelper::GenerateAutomaticCast(Statement.LHS->Type, Statement.RHS[0]->Type, BeginCast, EndCast);
+	FEmitHelper::GenerateAutomaticCast(EmitterContext, Statement.LHS->Type, Statement.RHS[0]->Type, BeginCast, EndCast);
 	EmitterContext.AddLine(FString::Printf(TEXT("%s = %s%s%s;"), *DestinationExpression, *BeginCast, *SourceExpression, *EndCast));
 }
 
@@ -391,7 +391,7 @@ FString FBlueprintCompilerCppBackend::EmitMethodInputParameterList(FEmitterLocal
 				check(Schema);
 				if (Schema->ConvertPropertyToPinType(FuncParamProperty, LType))
 				{
-					FEmitHelper::GenerateAutomaticCast(LType, Term->Type, BeginCast, CloseCast);
+					FEmitHelper::GenerateAutomaticCast(EmitterContext, LType, Term->Type, BeginCast, CloseCast);
 				}
 				VarName += BeginCast;
 				VarName += TermToText(EmitterContext, Term);
@@ -432,7 +432,7 @@ FString FBlueprintCompilerCppBackend::EmitCallStatmentInner(FEmitterLocalContext
 			check(Schema);
 			if (Schema->ConvertPropertyToPinType(FuncToCallReturnProperty, RType))
 			{
-				FEmitHelper::GenerateAutomaticCast(Statement.LHS->Type, RType, BeginCast, CloseCast);
+				FEmitHelper::GenerateAutomaticCast(EmitterContext, Statement.LHS->Type, RType, BeginCast, CloseCast);
 			}
 			Result += FString::Printf(TEXT("%s = %s"), *TermToText(EmitterContext, Statement.LHS), *BeginCast);
 		}
diff --git a/Engine/Source/Developer/BlueprintCompilerCppBackend/Private/BlueprintCompilerCppBackendUtils.cpp b/Engine/Source/Developer/BlueprintCompilerCppBackend/Private/BlueprintCompilerCppBackendUtils.cpp
--- a/Engine/Source/Developer/BlueprintCompilerCppBackend/Private/BlueprintCompilerCppBackendUtils.cpp
+++ b/Engine/Source/Developer/BlueprintCompilerCppBackend/Private/BlueprintCompilerCppBackendUtils.cpp
@@ -955,7 +955,7 @@ bool FEmitHelper::ShouldHandleAsImplementableEvent(UFunction* Function)
 	return false;
 }
 
-bool FEmitHelper::GenerateAutomaticCast(const FEdGraphPinType& LType, const FEdGraphPinType& RType, FString& OutCastBegin, FString& OutCastEnd)
+bool FEmitHelper::GenerateAutomaticCast(FEmitterLocalContext& EmitterContext, const FEdGraphPinType& LType, const FEdGraphPinType& RType, FString& OutCastBegin, FString& OutCastEnd)
 {
 	// BYTE to ENUM cast
 	// ENUM to BYTE cast
@@ -983,6 +983,7 @@ bool FEmitHelper::GenerateAutomaticCast(const FEdGraphPinType& LType, const FEdG
 	if (LType.PinCategory == UEdGraphSchema_K2::PC_Object)
 	{
 		UClass* LClass = Cast<UClass>(LType.PinSubCategoryObject.Get());
+		LClass = LClass ? EmitterContext.GetFirstNativeOrConvertedClass(LClass) : nullptr;
 		UClass* RClass = Cast<UClass>(RType.PinSubCategoryObject.Get());
 		if (LClass && RClass && LClass->IsChildOf(RClass) && !RClass->IsChildOf(LClass))
 		{
diff --git a/Engine/Source/Developer/BlueprintCompilerCppBackend/Private/BlueprintCompilerCppBackendUtils.h b/Engine/Source/Developer/BlueprintCompilerCppBackend/Private/BlueprintCompilerCppBackendUtils.h
--- a/Engine/Source/Developer/BlueprintCompilerCppBackend/Private/BlueprintCompilerCppBackendUtils.h
+++ b/Engine/Source/Developer/BlueprintCompilerCppBackend/Private/BlueprintCompilerCppBackendUtils.h
@@ -221,7 +221,7 @@ struct FEmitHelper
 
 	static bool ShouldHandleAsImplementableEvent(UFunction* Function);
 
-	static bool GenerateAutomaticCast(const FEdGraphPinType& LType, const FEdGraphPinType& RType, FString& OutCastBegin, FString& OutCastEnd);
+	static bool GenerateAutomaticCast(FEmitterLocalContext& EmitterContext, const FEdGraphPinType& LType, const FEdGraphPinType& RType, FString& OutCastBegin, FString& OutCastEnd);
 
 	static FString GenerateReplaceConvertedMD(UObject* Obj);
 
