CIS fix: Non-Unity compile issue, was using class restricted to the implementation file.

diff --git a/Engine/Source/Developer/BlueprintCompilerCppBackend/Private/BlueprintCompilerCppBackendUtils.cpp b/Engine/Source/Developer/BlueprintCompilerCppBackend/Private/BlueprintCompilerCppBackendUtils.cpp
--- a/Engine/Source/Developer/BlueprintCompilerCppBackend/Private/BlueprintCompilerCppBackendUtils.cpp
+++ b/Engine/Source/Developer/BlueprintCompilerCppBackend/Private/BlueprintCompilerCppBackendUtils.cpp
@@ -580,7 +580,7 @@ FString FEmitHelper::GetBaseFilename(const UObject* AssetObj)
 FString FEmitHelper::GetPCHFilename()
 {
 	FString PCHFilename;
-	FBlueprintCompilerCppBackendModule& BackEndModule = (FBlueprintCompilerCppBackendModule&)FBlueprintCompilerCppBackendModule::Get();
+	IBlueprintCompilerCppBackendModule& BackEndModule = (IBlueprintCompilerCppBackendModule&)IBlueprintCompilerCppBackendModule::Get();
 
 	auto& PchFilenameQuery = BackEndModule.OnPCHFilenameQuery();
 	if (PchFilenameQuery.IsBound())
