Moving these BP compiler tests into engine content so that I can iterate on them in the dev-blueprints stream

diff --git a/Engine/Source/Developer/BlueprintCompilerCppBackend/Private/BPCompilerTests.cpp b/Engine/Source/Developer/BlueprintCompilerCppBackend/Private/BPCompilerTests.cpp
new file mode 100644
--- /dev/null
+++ b/Engine/Source/Developer/BlueprintCompilerCppBackend/Private/BPCompilerTests.cpp
@@ -0,0 +1,409 @@
+// Copyright 1998-2015 Epic Games, Inc. All Rights Reserved.
+
+#include "BlueprintCompilerCppBackendModulePrivatePCH.h"
+
+// Helper functions used for verifying that objects are identical, objects are expected to have a different class
+// but identical properties (transient data skipped, treated as an implementation detail):
+bool Identical(const UObject* A, const UObject* B)
+{
+	if (!A && !B)
+	{
+		return true;
+	}
+	
+	if (!A || !B)
+	{
+		return false;
+	}
+
+	bool bResult = true;
+	TFieldIterator<UProperty> ItA(A->GetClass());
+	TFieldIterator<UProperty> ItB(B->GetClass());
+	while( (ItA || ItB) && bResult )
+	{
+		const UProperty* PropA = ItA ? *ItA : nullptr;
+		const UProperty* PropB = ItB ? *ItB : nullptr;
+
+		// skip transient structs (ie ubergraph frame):
+		bool bDone = false;
+
+		// property is transient or is owned by a structure that is always transient:
+		if (PropA && PropA->HasAllPropertyFlags(CPF_Transient))
+		{
+			++ItA;
+			bDone = true;
+		}
+		if (PropB && PropB->HasAllPropertyFlags(CPF_Transient))
+		{
+			++ItB;
+			bDone = true;
+		}
+
+		// for now, assume properties occur in the same order:
+		if (!bDone)
+		{
+			if (PropA == nullptr || PropB == nullptr)
+			{
+				bResult = false;
+			}
+			else if (PropA->GetName() == PropB->GetName())
+			{
+				const void* AValue = PropA->ContainerPtrToValuePtr<void>(A);
+				const void* BValue = PropB->ContainerPtrToValuePtr<void>(B);
+
+				const UClassProperty* ClassPropA = Cast<UClassProperty>(PropA);
+				const UClassProperty* ClassPropB = Cast<UClassProperty>(PropB);
+				if (ClassPropA && ClassPropB)
+				{
+					// make sure we treat generated and native instances as equivalent:
+					UObject* ObjectA = ClassPropA->GetObjectPropertyValue(AValue);
+					UObject* ObjectB = ClassPropB->GetObjectPropertyValue(BValue);
+
+					bResult = ObjectA && ObjectB && ObjectA->GetName() == ObjectB->GetName();
+				}
+				else
+				{
+					bResult = PropA->Identical(AValue, BValue, PPF_DeepComparison);
+				}
+			}
+			else
+			{
+				bResult = false;
+			}
+
+			if (ItA)
+			{
+				++ItA;
+			}
+
+			if (ItB)
+			{
+				++ItB;
+			}
+		}
+	}
+
+	return bResult && !ItA && !ItB;
+}
+
+// Helper functions introduced to load classes (generated or native:
+static UClass* GetGeneratedClass(const TCHAR* TestFolder, const TCHAR* ClassName, FAutomationTestBase* Context)
+{
+	FString FullName = FString::Printf(TEXT("/Engine/NotForLicensees/Automation/CompilerTests/%s/%s.%s"), TestFolder, ClassName, ClassName);
+	UBlueprint* Blueprint = Cast<UBlueprint>(StaticLoadObject(UBlueprint::StaticClass(), nullptr, *FullName));
+
+	if (!Blueprint)
+	{
+		Context->AddWarning(FString::Printf(TEXT("Missing blueprint for test: '%s'"), *FullName));
+		return nullptr;
+	}
+
+	return Blueprint->GeneratedClass;
+}
+
+static UClass* GetNativeClass(const TCHAR* ClassName, FAutomationTestBase* Context)
+{
+	UPackage* NativePackage = FindObjectFast<UPackage>(nullptr, TEXT("/Script/Engine"));
+	check(NativePackage);
+	UClass* Ret = FindObjectFast<UClass>(NativePackage, ClassName);
+	if (!Ret)
+	{
+		Context->AddWarning(FString::Printf(TEXT("Missing native type for test: '%s'"), ClassName));
+		return nullptr;
+	}
+	return Ret;
+}
+
+// Helper functions introduced to avoid crashing if classes are missing:
+static UObject* NewTestObject(UClass* Class)
+{
+	if (Class)
+	{
+		return NewObject<UObject>(GetTransientPackage(), Class);
+	}
+	return nullptr;
+}
+
+static UObject* NewTestActor(UClass* ActorClass)
+{
+	if (ActorClass)
+	{
+		AActor* Actor = GWorld->SpawnActor(ActorClass);
+#if WITH_EDITORONLY_DATA
+		Actor->GetRootComponent()->bVisualizeComponent = true;
+#endif
+	}
+	return nullptr;
+}
+
+static void Call(UObject* Target, const TCHAR* FunctionName, void* Args = nullptr)
+{
+	if (!Target)
+	{
+		return;
+	}
+	UFunction* Fn = Target->FindFunction(FunctionName);
+	if (Fn)
+	{
+		Target->ProcessEvent(Fn, Args);
+	}
+}
+
+// Remove EAutomationTestFlags::Disabled to enable these tests, note that these will need to be moved into the ClientContext because we can only test cooked content:
+static const uint32 CompilerTestFlags = EAutomationTestFlags::EditorContext | EAutomationTestFlags::EngineFilter;
+
+// Tests:
+IMPLEMENT_SIMPLE_AUTOMATION_TEST(FBPCompilerArrayTest, "Project.Blueprints.NativeBackend.ArrayTest", CompilerTestFlags)
+bool FBPCompilerArrayTest::RunTest(const FString& Parameters)
+{
+	// setup input data:
+	TArray<FString> InputA;
+	TArray<FString> InputB;
+	TArray<FString> Input;
+	{
+		Input.Push(TEXT("addedString"));
+		InputA = Input;
+		InputB = Input;
+	}
+	UObject* GeneratedTestInstance = NewTestObject( GetGeneratedClass(TEXT("Array"), TEXT("BP_Array_Basic"), this));
+	UObject* NativeTestInstance = NewTestObject( GetNativeClass(TEXT("BP_Array_Basic_C"), this));
+
+	// execute test on generated object and native object:
+	Call(GeneratedTestInstance, TEXT("RunArrayTest"), &InputA);
+	Call(NativeTestInstance, TEXT("RunArrayTest"), &InputB);
+
+	// evaluate results:
+	if (InputA != Input)
+	{
+		AddError(FString::Printf(TEXT("%S: Bytecode altered input"), *GeneratedTestInstance->GetName()));
+	}
+	if (InputB != Input)
+	{
+		AddError(FString::Printf(TEXT("%S: Native code altered input"), *NativeTestInstance->GetName()));
+	}
+	if (!Identical(GeneratedTestInstance, NativeTestInstance))
+	{
+		AddError(FString::Printf(TEXT("%S: Generated instance differs from native instance"), *GeneratedTestInstance->GetName()));
+	}
+
+	return true;
+} 
+
+IMPLEMENT_SIMPLE_AUTOMATION_TEST(FBPCompilerCDOTest, "Project.Blueprints.NativeBackend.CDOTest", CompilerTestFlags)
+bool FBPCompilerCDOTest::RunTest(const FString& Parameters)
+{
+	// setup input data:
+	const UObject* CDOGenerated = nullptr;
+	const UObject* CDONative = nullptr;
+	{
+		UObject* GeneratedTestInstance = NewTestObject(GetGeneratedClass(TEXT("CDO"), TEXT("BP_CDO_Basic"), this));
+		UObject* NativeTestInstance = NewTestObject(GetNativeClass(TEXT("BP_CDO_Basic_C"), this));
+
+		if (!GeneratedTestInstance || !NativeTestInstance)
+		{
+			return true;
+		}
+
+		CDOGenerated = GeneratedTestInstance->GetClass()->ClassDefaultObject;
+		CDONative = NativeTestInstance->GetClass()->ClassDefaultObject;
+	}
+
+	// evaluate results:
+	if (!Identical(CDOGenerated, CDONative))
+	{
+		AddError(FString::Printf(TEXT("%S: Generated CDO differs from native CDO"), *CDOGenerated->GetName()));
+	}
+	return true;
+}
+
+IMPLEMENT_SIMPLE_AUTOMATION_TEST(FBPCompilerCommunicationTest, "Project.Blueprints.NativeBackend.CommunicationTest", CompilerTestFlags)
+bool FBPCompilerCommunicationTest::RunTest(const FString& Parameters)
+{
+	// setup input data:
+	UObject* GeneratedTestInstanceA = NewTestObject(GetGeneratedClass(TEXT("Communication"), TEXT("BP_Comm_Test_A"), this));
+	UObject* GeneratedTestInstanceB = NewTestObject(GetGeneratedClass(TEXT("Communication"), TEXT("BP_Comm_Test_B"), this));
+	if (!GeneratedTestInstanceA || !GeneratedTestInstanceB)
+	{
+		return true;
+	}
+
+	UObject* NativeTestInstanceA = NewTestObject(GetNativeClass(TEXT("BP_Comm_Test_A_C"), this));
+	UObject* NativeTestInstanceB = NewTestObject(GetNativeClass(TEXT("BP_Comm_Test_B_C"), this));
+
+	// execute test on generated object and native object:
+	auto Test = [](UObject* A, UObject* B)
+	{
+		Call(A, TEXT("Flop"));
+		Call(B, TEXT("Flip"));
+	};
+	Test(GeneratedTestInstanceA, GeneratedTestInstanceB);
+	Test(NativeTestInstanceA, NativeTestInstanceB);
+
+	// evaluate results:
+	if (!Identical(GeneratedTestInstanceA, NativeTestInstanceA))
+	{
+		AddError(TEXT("Different results for communication test object A"));
+	}
+	if (!Identical(GeneratedTestInstanceB, NativeTestInstanceB))
+	{
+		AddError(TEXT("Different results for communication test object B"));
+	}
+	
+	return true;
+}
+
+IMPLEMENT_SIMPLE_AUTOMATION_TEST(FBPCompilerConstructionScriptTest, "Project.Blueprints.NativeBackend.ConstructionScriptTest", CompilerTestFlags)
+bool FBPCompilerConstructionScriptTest::RunTest(const FString& Parameters)
+{
+	// setup input data:
+	UObject* GeneratedTestInstance = NewTestActor(GetGeneratedClass(TEXT("ConstructionScript"), TEXT("BP_ConstructionScript_Test"), this));
+	UObject* NativeTestInstance = NewTestActor(GetNativeClass(TEXT("BP_ConstructionScript_Test_C"), this));
+
+	// evaluate results:
+	if (!Identical(GeneratedTestInstance, NativeTestInstance))
+	{
+		AddError(FString::Printf(TEXT("%S: Generated instance differs from native instance"), *GeneratedTestInstance->GetName()));
+	}
+	return true;
+}
+
+IMPLEMENT_SIMPLE_AUTOMATION_TEST(FBPCompilerControlFlowTest, "Project.Blueprints.NativeBackend.ControlFlow", CompilerTestFlags)
+bool FBPCompilerControlFlowTest::RunTest(const FString& Parameters)
+{
+	// setup input data:
+	UObject* GeneratedTestInstance = NewTestObject(GetGeneratedClass(TEXT("ControlFlow"), TEXT("BP_ControlFlow_Basic"), this));
+	UObject* NativeTestInstance = NewTestObject(GetNativeClass(TEXT("BP_ControlFlow_Basic_C"), this));
+
+	// execute test on generated object and native object:
+	Call(GeneratedTestInstance, TEXT("RunControlFlowTest"));
+	Call(NativeTestInstance, TEXT("RunControlFlowTest"));
+
+	// evaluate results:
+	if (!Identical(GeneratedTestInstance, NativeTestInstance))
+	{
+		AddError(FString::Printf(TEXT("%S: Generated instance differs from native instance"), *GeneratedTestInstance->GetName()));
+	}
+	return true;
+}
+
+IMPLEMENT_SIMPLE_AUTOMATION_TEST(FBPCompilerEnumTest, "Project.Blueprints.NativeBackend.EnumTest", CompilerTestFlags)
+bool FBPCompilerEnumTest::RunTest(const FString& Parameters)
+{
+	// setup input data:
+	UObject* GeneratedTestInstance = NewTestObject(GetGeneratedClass(TEXT("Enum"), TEXT("BP_Enum_Reader_Writer"), this));
+	UObject* NativeTestInstance = NewTestObject(GetNativeClass(TEXT("BP_Enum_Reader_Writer_C"), this));
+
+	// execute test on generated object and native object:
+	Call(GeneratedTestInstance, TEXT("UpdateEnum"));
+	Call(NativeTestInstance, TEXT("UpdateEnum"));
+
+	// evaluate results:
+	if (!Identical(GeneratedTestInstance, NativeTestInstance))
+	{
+		AddError(FString::Printf(TEXT("%S: Generated instance differs from native instance"), *GeneratedTestInstance->GetName()));
+	}
+	return true;
+}
+
+IMPLEMENT_SIMPLE_AUTOMATION_TEST(FBPCompilerEventTest, "Project.Blueprints.NativeBackend.Event", CompilerTestFlags)
+bool FBPCompilerEventTest::RunTest(const FString& Parameters)
+{
+	// setup input data:
+	UObject* GeneratedTestInstance = NewTestObject(GetGeneratedClass(TEXT("Event"), TEXT("BP_Event_Basic"), this));
+	UObject* NativeTestInstance = NewTestObject(GetNativeClass(TEXT("BP_Event_Basic_C"), this));
+
+	// execute test on generated object and native object:
+	Call(GeneratedTestInstance, TEXT("BeginEventChain"));
+	Call(NativeTestInstance, TEXT("BeginEventChain"));
+
+	// evaluate results:
+	if (!Identical(GeneratedTestInstance, NativeTestInstance))
+	{
+		AddError(FString::Printf(TEXT("%S: Generated instance differs from native instance"), *GeneratedTestInstance->GetName()));
+	}
+	return true;
+}
+
+struct FInheritenceTestParams
+{
+	bool bFlag;
+	TArray<FString> Strings;
+	TArray<int32> Result;
+};
+
+inline bool operator==(const FInheritenceTestParams& LHS, const FInheritenceTestParams& RHS)
+{
+	return 
+		LHS.bFlag == RHS.bFlag &&
+		LHS.Strings == RHS.Strings &&
+		LHS.Result == RHS.Result;
+}
+
+inline bool operator!=(const FInheritenceTestParams& LHS, const FInheritenceTestParams& RHS)
+{
+	return !(LHS == RHS);
+}
+
+IMPLEMENT_SIMPLE_AUTOMATION_TEST(FBPCompilerInheritenceTest, "Project.Blueprints.NativeBackend.Inheritence", CompilerTestFlags)
+bool FBPCompilerInheritenceTest::RunTest(const FString& Parameters)
+{
+	// setup input data:
+	UObject* GeneratedTestInstance = NewTestObject(GetGeneratedClass(TEXT("Inheritence"), TEXT("BP_Child_Basic"), this));
+	UObject* NativeTestInstance = NewTestObject(GetNativeClass(TEXT("BP_Child_Basic_C"), this));
+
+	// execute test on generated object and native object:
+	FInheritenceTestParams ParamsGenerated;
+	ParamsGenerated.bFlag = true;
+	FInheritenceTestParams ParamsNative = ParamsGenerated;
+	Call(GeneratedTestInstance, TEXT("VirtualFunction"), &ParamsGenerated);
+	Call(NativeTestInstance, TEXT("VirtualFunction"), &ParamsNative);
+
+	// evaluate results:
+	if (ParamsGenerated != ParamsNative)
+	{
+		AddError(FString::Printf(TEXT("%S: Outputs differ for native and generated instance"), *GeneratedTestInstance->GetName()));
+	}
+	if (!Identical(GeneratedTestInstance, NativeTestInstance))
+	{
+		AddError(FString::Printf(TEXT("%S: Generated instance differs from native instance"), *GeneratedTestInstance->GetName()));
+	}
+	return true;
+}
+
+IMPLEMENT_SIMPLE_AUTOMATION_TEST(FBPCompilerStructureTest, "Project.Blueprints.NativeBackend.Structure", CompilerTestFlags)
+bool FBPCompilerStructureTest::RunTest(const FString& Parameters)
+{
+	// setup input data:
+	UObject* GeneratedTestInstance = NewTestObject(GetGeneratedClass(TEXT("Structure"), TEXT("BP_Structure_Driver"), this));
+	UObject* NativeTestInstance = NewTestObject(GetNativeClass(TEXT("BP_Structure_Driver_C"), this));
+
+	// execute test on generated object and native object:
+	Call(GeneratedTestInstance, TEXT("RunStructTest"));
+	Call(NativeTestInstance, TEXT("RunStructTest"));
+
+	// evaluate results:
+	if (!Identical(GeneratedTestInstance, NativeTestInstance))
+	{
+		AddError(FString::Printf(TEXT("%S: Generated instance differs from native instance"), *GeneratedTestInstance->GetName()));
+	}
+	return true;
+}
+
+IMPLEMENT_SIMPLE_AUTOMATION_TEST(FBPCompilerNodeTest, "Project.Blueprints.NativeBackend.Node", CompilerTestFlags)
+bool FBPCompilerNodeTest::RunTest(const FString& Parameters)
+{
+	// setup input data:
+	UObject* GeneratedTestInstance = NewTestObject(GetGeneratedClass(TEXT("Node"), TEXT("BP_Node_Basic"), this));
+	UObject* NativeTestInstance = NewTestObject(GetNativeClass(TEXT("BP_Node_Basic_C"), this));
+
+	// execute test on generated object and native object:
+	Call(GeneratedTestInstance, TEXT("RunNodes"));
+	Call(NativeTestInstance, TEXT("RunNodes"));
+
+	// evaluate results:
+	if (!Identical(GeneratedTestInstance, NativeTestInstance))
+	{
+		AddError(FString::Printf(TEXT("%S: Generated instance differs from native instance"), *GeneratedTestInstance->GetName()));
+	}
+	return true;
+}
