Copying //UE4/Dev-Blueprints to Dev-Main (//UE4/Dev-Main) @ 2781164

==========================
MAJOR FEATURES + CHANGES
==========================

Change 2716841 on 2015/10/05 by [REDACTED_PERSON]

	(WIP) Cleaning up how we setup script assets for replacement on cook (aligning more with the Blueprint conversion tool).

Change 2719089 on 2015/10/07 by [REDACTED_PERSON]

	ToValidCPPIdentifierChars handles propertly '?' char.

Change 2719361 on 2015/10/07 by [REDACTED_PERSON]

	Generated native code for AnimBPGC - some preliminary changes.

	Refactor: UAnimBlueprintGeneratedClass is not accessed directly in runtime. It is accessed via UAnimClassInterface interface.
	Properties USkeletalMeshComponent::AnimBlueprintGeneratedClass and UInterpTrackFloatAnimBPParam::AnimBlueprintClass were changed into "TSubclassOf<UAnimInstance> AnimClass"

	The UDynamicClass also can deliver the IAnimClassInterface interface. See UAnimClassData, IAnimClassInterface::GetFromClass and UDynamicClass::AnimClassImplementation.

Change 2719383 on 2015/10/07 by [REDACTED_PERSON]

	Debug-only code removed

Change 2720528 on 2015/10/07 by [REDACTED_PERSON]

	Fix for determinsitc cooking of async tasks and load asset nodes

Change 2721273 on 2015/10/08 by [REDACTED_PERSON]

	Blueprint Compiler Cpp Backend
	- Anim Blueprints can be converted
	- Various fixes/improvements

Change 2721310 on 2015/10/08 by [REDACTED_PERSON]

	refactor (cl#2719361) - no "auto" keyword

Change 2721727 on 2015/10/08 by [REDACTED_PERSON]

	(WIP) Setup the cook commandlet so it handles converted assets, replacing them with generated classes.

	    - Refactored the conversion manifest (using a map over an array)
	    - Centralized destination paths into a helper struct (for the manifest)
	    - Generating an Editor module that automatically hooks into the cook process when enabled
	    - Loading and applying native replacments for the cook

Change 2723276 on 2015/10/09 by [REDACTED_PERSON]

	Blueprints duplicated for PIE will no longer register as dependencies to other Blueprint.

Change 2724345 on 2015/10/11 by [REDACTED_PERSON]

	Blueprint profiler at first pass, this includes the ability to instrument specific blueprints with realtime editor stat display.
	#UEBP-21 - Profiling data capture and storage
	#UEBP-13 - Performance capture landing page
	#Branch UE4
	#Proj BlueprintProfiler, BlueprintGraph, EditorStyle, Kismet, UnrealEd, CoreUObject, Engine

Change 2724613 on 2015/10/12 by [REDACTED_PERSON]

	Incremental update for blueprint profiler to fix the way some of the reported stats cascade through events and branches and additionally some missed bits of code are refactored/removed.
	#Branch UE4
	#Proj BlueprintProfiler

	#info Whilst looking into this I spotted the reason why the stats seem so erratic, There appears to be an issue with FText's use of EXPERIMENTAL_TEXT_FAST_DECIMAL_FORMAT which I have reported, but ideally disable this locally until a fix is integrated.

Change 2724723 on 2015/10/12 by [REDACTED_PERSON]

	Constructor of a dynamic class creates CDO.

Change 2725108 on 2015/10/12 by [REDACTED_PERSON]

	[UE-21891] Minor fix to the array shuffle() function; now processes the last entry like all the others.

Change 2726358 on 2015/10/13 by [REDACTED_PERSON]

	UDataTable is properly saved even if its RowStruct is null.

	https://udn.unrealengine.com/questions/264064/crash-using-hotreload-in-custom-datatable-cdo-clas.html

Change 2727395 on 2015/10/13 by [REDACTED_PERSON]

	(WIP) Second pass on the Blueprint conversion pipeline; setting it up for more optimal (speedier) performance.
	    * Using stubs for replacements (rather than loading dynamic replacement).
	    * Giving the cook commandlet more control (so a conversion could be ran directly).
	    * Now logging replacements by old object path (to account for UPackage replacement queries).
	    * Fix for [UE-21947], unshelved from CL 2724944 (by [REDACTED_PERSON]).

Change 2727484 on 2015/10/13 by [REDACTED_PERSON]

	[UE-22008] Fixing up comment/tooltip typo for UActorComponent::bAutoActivate.

Change 2727527 on 2015/10/13 by [REDACTED_PERSON]

	Downgrading an inactionable EdGraph warning, while adding more info so we could possibly determine what's happening.

Change 2727702 on 2015/10/13 by [REDACTED_PERSON]

	Fix for crash in UDelegateProperty::GetCPPType when called on a function with no OwnerClass (events)

Change 2727968 on 2015/10/14 by [REDACTED_PERSON]

	Since ConstructorHelpers::FClassFinder is usually static, the loaded class should be in root set, to prevent the pointer stored in  ConstructorHelpers::FClassFinder from being obsolete.
	FindOrLoadClass behaves now like FindOrLoadObject.

Change 2728139 on 2015/10/14 by [REDACTED_PERSON]

diff --git a/Engine/Config/BaseEditor.ini b/Engine/Config/BaseEditor.ini
--- a/Engine/Config/BaseEditor.ini
+++ b/Engine/Config/BaseEditor.ini
@@ -463,3 +463,6 @@ AngleUnits=EUnit::Degrees
 SpeedUnits=EUnit::MetersPerSecond
 TemperatureUnits=EUnit::Celsius
 ForceUnits=EUnit::Newtons
+
+[/Script/BlueprintNativeCodeGen.BlueprintNativeCodeGenConfig]
++ExcludedBlueprintTypes=/Script/Engine.LevelScriptActor
diff --git a/Engine/Config/BaseEditorPerProjectUserSettings.ini b/Engine/Config/BaseEditorPerProjectUserSettings.ini
--- a/Engine/Config/BaseEditorPerProjectUserSettings.ini
+++ b/Engine/Config/BaseEditorPerProjectUserSettings.ini
@@ -678,17 +678,16 @@ DefaultBrushRadius=128
 +Node=(Class=KismetArrayLibrary:Array_Get Key=A Shift=false Ctrl=false Alt=false)
 
 [DefaultEventNodes]
++Node=(TargetClass=Actor TargetEvent="ReceiveBeginPlay")
++Node=(TargetClass=Actor TargetEvent="ReceiveActorBeginOverlap")
++Node=(TargetClass=Actor TargetEvent="ReceiveTick")
 +Node=(TargetClass=ActorComponent TargetEvent="ReceiveBeginPlay")
 +Node=(TargetClass=ActorComponent TargetEvent="ReceiveTick")
 +Node=(TargetClass=GameplayAbility TargetEvent="K2_ActivateAbility")
++Node=(TargetClass=GameplayAbility TargetEvent="K2_OnEndAbility")
 +Node=(TargetClass=UserWidget TargetEvent="Construct")
 +Node=(TargetClass=UserWidget TargetEvent="Tick")
 +Node=(TargetClass=AnimInstance TargetEvent="BlueprintUpdateAnimation")
-+Node=(TargetClass=OrionGameplayCueNotify_Looping TargetEvent="K2_OnActive")
-+Node=(TargetClass=OrionGameplayCueNotify_Looping TargetEvent="K2_WhileActive")
-+Node=(TargetClass=OrionGameplayCueNotify_Looping TargetEvent="K2_OnExecute")
-+Node=(TargetClass=OrionGameplayCueNotify_Looping TargetEvent="K2_OnRemove")
-+Node=(TargetClass=OrionGameplayCueNotify_Burst_Latent TargetEvent="K2_OnExecute")
 
 [MaterialEditorSpawnNodes]
 +Node=(Class=MaterialExpressionAdd Key=A Shift=false Ctrl=false Alt=false)
diff --git a/Engine/Config/BaseEngine.ini b/Engine/Config/BaseEngine.ini
--- a/Engine/Config/BaseEngine.ini
+++ b/Engine/Config/BaseEngine.ini
@@ -1724,6 +1724,7 @@ bExecutionAfterReturn=false
 bCanSuppressAccessViolation=false
 bUseSelectRef=true
 bLoadNativeConvertedBPClassInEditor=false
+bIgnoreCompileOnLoadErrorsOnBuildMachine=false
 
 [/Script/Engine.Blueprint]
 bRecompileOnLoad=true
diff --git a/Engine/Source/Developer/BlueprintCompilerCppBackend/Private/BPCompilerTests.cpp b/Engine/Source/Developer/BlueprintCompilerCppBackend/Private/BPCompilerTests.cpp
--- a/Engine/Source/Developer/BlueprintCompilerCppBackend/Private/BPCompilerTests.cpp
+++ b/Engine/Source/Developer/BlueprintCompilerCppBackend/Private/BPCompilerTests.cpp
@@ -2,95 +2,72 @@
 
 #include "BlueprintCompilerCppBackendModulePrivatePCH.h"
 
-// Helper functions used for verifying that objects are identical, objects are expected to have a different class
-// but identical properties (transient data skipped, treated as an implementation detail):
-bool Identical(const UObject* A, const UObject* B)
+class FArchiveSkipTransientObjectCRC32 : public FArchiveObjectCrc32
 {
-	if (!A && !B)
+	// Begin FArchive Interface
+	virtual bool ShouldSkipProperty(const UProperty* InProperty) const override
 	{
-		return true;
+		return FArchiveObjectCrc32::ShouldSkipProperty(InProperty)
+			|| InProperty->HasAllPropertyFlags(CPF_Transient)
+			|| InProperty->HasAllPropertyFlags(CPF_EditorOnly)
+			|| InProperty->GetName() == TEXT("BlueprintCreatedComponents")
+			|| InProperty->GetName() == TEXT("CreationMethod")
+			|| InProperty->GetName() == TEXT("InstanceComponents")
+			|| InProperty->GetName() == TEXT("bNetAddressable")
+			|| InProperty->GetName() == TEXT("OwnedComponents");
 	}
-	
-	if (!A || !B)
+	// End FArchive Interface
+};
+
+void ClearNativeRecursive(UObject* OnObject)
+{
+	OnObject->ClearInternalFlags(EInternalObjectFlags::Native);
+	TArray<UObject*> Children;
+	GetObjectsWithOuter(OnObject, Children, false);
+	for (auto Entry : Children)
 	{
-		return false;
+		ClearNativeRecursive(Entry);
 	}
+}
 
-	bool bResult = true;
-	TFieldIterator<UProperty> ItA(A->GetClass());
-	TFieldIterator<UProperty> ItB(B->GetClass());
-	while( (ItA || ItB) && bResult )
+// We mess with the rootset flag instead of using FGCObject because it's an error for any test data to remain in the rootset after the test runs
+struct FOwnedObjectsHelper
+{
+	~FOwnedObjectsHelper()
 	{
-		const UProperty* PropA = ItA ? *ItA : nullptr;
-		const UProperty* PropB = ItB ? *ItB : nullptr;
-
-		// skip transient structs (ie ubergraph frame):
-		bool bDone = false;
-
-		// property is transient or is owned by a structure that is always transient:
-		if (PropA && PropA->HasAllPropertyFlags(CPF_Transient))
-		{
-			++ItA;
-			bDone = true;
-		}
-		if (PropB && PropB->HasAllPropertyFlags(CPF_Transient))
-		{
-			++ItB;
-			bDone = true;
-		}
-
-		// for now, assume properties occur in the same order:
-		if (!bDone)
+		for (auto Entry : OwnedObjects)
 		{
-			if (PropA == nullptr || PropB == nullptr)
-			{
-				bResult = false;
-			}
-			else if (PropA->GetName() == PropB->GetName())
-			{
-				const void* AValue = PropA->ContainerPtrToValuePtr<void>(A);
-				const void* BValue = PropB->ContainerPtrToValuePtr<void>(B);
-
-				const UClassProperty* ClassPropA = Cast<UClassProperty>(PropA);
-				const UClassProperty* ClassPropB = Cast<UClassProperty>(PropB);
-				if (ClassPropA && ClassPropB)
-				{
-					// make sure we treat generated and native instances as equivalent:
-					UObject* ObjectA = ClassPropA->GetObjectPropertyValue(AValue);
-					UObject* ObjectB = ClassPropB->GetObjectPropertyValue(BValue);
-
-					bResult = ObjectA && ObjectB && ObjectA->GetName() == ObjectB->GetName();
-				}
-				else
-				{
-					bResult = PropA->Identical(AValue, BValue, PPF_DeepComparison);
-				}
-			}
-			else
+			Entry->RemoveFromRoot();
+			Entry->ClearFlags(RF_Standalone);
+			if (Entry->IsNative())
 			{
-				bResult = false;
+				ClearNativeRecursive(Entry);
 			}
 
-			if (ItA)
+			// Actors need to be explicitly destroyed, probably just to remove them from their owning
+			// level:
+			if (AActor* AsActor = Cast<AActor>(Entry))
 			{
-				++ItA;
-			}
-
-			if (ItB)
-			{
-				++ItB;
+				AsActor->Destroy();
 			}
 		}
+		CollectGarbage(RF_NoFlags);
 	}
 
-	return bResult && !ItA && !ItB;
-}
+	void Push(UObject* Obj)
+	{
+		Obj->AddToRoot();
+		OwnedObjects.Push(Obj);
+	}
+private:
+	TArray<UObject*> OwnedObjects;
+};
 
 // Helper functions introduced to load classes (generated or native:
-static UClass* GetGeneratedClass(const TCHAR* TestFolder, const TCHAR* ClassName, FAutomationTestBase* Context)
+static UClass* GetGeneratedClass(const TCHAR* TestFolder, const TCHAR* ClassName, FAutomationTestBase* Context, FOwnedObjectsHelper &OwnedObjects)
 {
 	FString FullName = FString::Printf(TEXT("/Engine/NotForLicensees/Automation/CompilerTests/%s/%s.%s"), TestFolder, ClassName, ClassName);
-	UBlueprint* Blueprint = Cast<UBlueprint>(StaticLoadObject(UBlueprint::StaticClass(), nullptr, *FullName));
+	UBlueprint* Blueprint = LoadObject<UBlueprint>(NULL, *FullName);
 
 	if (!Blueprint)
 	{
@@ -98,40 +75,96 @@ static UClass* GetGeneratedClass(const TCHAR* TestFolder, const TCHAR* ClassName
 		return nullptr;
 	}
 
+	TArray<UObject*> Objects;
+	GetObjectsWithOuter(Blueprint->GetOuter(), Objects, false);
+	for (auto Entry : Objects)
+	{
+		OwnedObjects.Push(Entry);
+	}
+
+	CollectGarbage(RF_NoFlags);
+
 	return Blueprint->GeneratedClass;
 }
 
-static UClass* GetNativeClass(const TCHAR* ClassName, FAutomationTestBase* Context)
+static UClass* GetNativeClass(const TCHAR* TestFolder, const TCHAR* ClassName, FAutomationTestBase* Context, FOwnedObjectsHelper &OwnedObjects)
 {
-	UPackage* NativePackage = FindObjectFast<UPackage>(nullptr, TEXT("/Script/Engine"));
+	FString FullName = FString::Printf(TEXT("/Game/Blueprints/CompilerTests/%s/%s"), TestFolder, ClassName);
+	UPackage* NativePackage = CreatePackage(NULL, *FullName);
 	check(NativePackage);
-	UClass* Ret = FindObjectFast<UClass>(NativePackage, ClassName);
+
+	CollectGarbage(RF_NoFlags);
+	const FString FStringFullPathName = FString::Printf(TEXT("%s.%s_C"), *FullName, ClassName);
+	UClass* Ret = (UClass*)ConstructDynamicType(*FStringFullPathName); //FindObjectFast<UClass>(NativePackage, *(FString(ClassName) + TEXT("_C")));
+
 	if (!Ret)
 	{
 		Context->AddWarning(FString::Printf(TEXT("Missing native type for test: '%s'"), ClassName));
 		return nullptr;
 	}
+
+	TArray<UObject*> Objects;
+	GetObjectsWithOuter(NativePackage, Objects, false);
+	for (auto Entry : Objects)
+	{
+		OwnedObjects.Push(Entry);
+	}
+
+	CollectGarbage(RF_NoFlags);
+
 	return Ret;
 }
 
+typedef UClass* (*ClassAccessor)(const TCHAR*, const TCHAR*, FAutomationTestBase*, FOwnedObjectsHelper&);
+
+typedef uint32(*TestImpl)(ClassAccessor F, FAutomationTestBase* Context);
+
+// This pattern is repeated in each test, so I'm just writing a helper function rather than copy/pasting it:
+static bool RunTestHelper(TestImpl T, FAutomationTestBase* Context)
+{
+	uint32 ResultsGenerated = T(&GetGeneratedClass, Context);
+	uint32 ResultsNative = T(&GetNativeClass, Context);
+
+	if (ResultsGenerated == 0)
+	{
+		Context->AddError(TEXT("Test failed to run!"));
+	}
+	else if (ResultsGenerated != ResultsNative)
+	{
+		Context->AddError(TEXT("Native differs from generated!"));
+	}
+
+	return true;
+}
+
 // Helper functions introduced to avoid crashing if classes are missing:
-static UObject* NewTestObject(UClass* Class)
+static UObject* NewTestObject(UClass* Class, FOwnedObjectsHelper &OwnedObjects)
 {
 	if (Class)
 	{
-		return NewObject<UObject>(GetTransientPackage(), Class);
+		UObject* Result = NewObject<UObject>(GetTransientPackage(), Class, NAME_None);
+		if (Result)
+		{
+			OwnedObjects.Push(Result);
+		}
+		return Result;
 	}
 	return nullptr;
 }
 
-static UObject* NewTestActor(UClass* ActorClass)
+static UObject* NewTestActor(UClass* ActorClass, FOwnedObjectsHelper &OwnedObjects)
 {
 	if (ActorClass)
 	{
 		AActor* Actor = GWorld->SpawnActor(ActorClass);
 #if WITH_EDITORONLY_DATA
-		Actor->GetRootComponent()->bVisualizeComponent = true;
+		if (Actor)
+		{
+			OwnedObjects.Push(Actor);
+			Actor->GetRootComponent()->bVisualizeComponent = true;
+		}
 #endif
+		return Actor;
 	}
 	return nullptr;
 }
@@ -156,172 +189,156 @@ static const uint32 CompilerTestFlags = EAutomationTestFlags::EditorContext | EA
 IMPLEMENT_SIMPLE_AUTOMATION_TEST(FBPCompilerArrayTest, "Project.Blueprints.NativeBackend.ArrayTest", CompilerTestFlags)
 bool FBPCompilerArrayTest::RunTest(const FString& Parameters)
 {
-	// setup input data:
-	TArray<FString> InputA;
-	TArray<FString> InputB;
-	TArray<FString> Input;
+	auto TestBody = [](ClassAccessor F, FAutomationTestBase* Context)
 	{
+		FOwnedObjectsHelper OwnedObjects;
+
+		TArray<FString> Input;
 		Input.Push(TEXT("addedString"));
-		InputA = Input;
-		InputB = Input;
-	}
-	UObject* GeneratedTestInstance = NewTestObject( GetGeneratedClass(TEXT("Array"), TEXT("BP_Array_Basic"), this));
-	UObject* NativeTestInstance = NewTestObject( GetNativeClass(TEXT("BP_Array_Basic_C"), this));
 
-	// execute test on generated object and native object:
-	Call(GeneratedTestInstance, TEXT("RunArrayTest"), &InputA);
-	Call(NativeTestInstance, TEXT("RunArrayTest"), &InputB);
+		UObject* TestInstance = NewTestObject(F(TEXT("Array"), TEXT("BP_Array_Basic"), Context, OwnedObjects), OwnedObjects);
+		if (!TestInstance)
+		{
+			return 0u;
+		}
 
-	// evaluate results:
-	if (InputA != Input)
-	{
-		AddError(FString::Printf(TEXT("%S: Bytecode altered input"), *GeneratedTestInstance->GetName()));
-	}
-	if (InputB != Input)
-	{
-		AddError(FString::Printf(TEXT("%S: Native code altered input"), *NativeTestInstance->GetName()));
-	}
-	if (!Identical(GeneratedTestInstance, NativeTestInstance))
-	{
-		AddError(FString::Printf(TEXT("%S: Generated instance differs from native instance"), *GeneratedTestInstance->GetName()));
-	}
+		Call(TestInstance, TEXT("RunArrayTest"), &Input);
 
-	return true;
-} 
+		FArchiveSkipTransientObjectCRC32 Results;
+		return Results.Crc32(TestInstance);
+	};
+
+	return RunTestHelper(TestBody, this);
+}
 
 IMPLEMENT_SIMPLE_AUTOMATION_TEST(FBPCompilerCDOTest, "Project.Blueprints.NativeBackend.CDOTest", CompilerTestFlags)
 bool FBPCompilerCDOTest::RunTest(const FString& Parameters)
 {
-	// setup input data:
-	const UObject* CDOGenerated = nullptr;
-	const UObject* CDONative = nullptr;
+	auto TestBody = [](ClassAccessor F, FAutomationTestBase* Context)
 	{
-		UObject* GeneratedTestInstance = NewTestObject(GetGeneratedClass(TEXT("CDO"), TEXT("BP_CDO_Basic"), this));
-		UObject* NativeTestInstance = NewTestObject(GetNativeClass(TEXT("BP_CDO_Basic_C"), this));
+		FOwnedObjectsHelper OwnedObjects;
 
-		if (!GeneratedTestInstance || !NativeTestInstance)
+		UObject* TestInstance = NewTestObject(F(TEXT("CDO"), TEXT("BP_CDO_Basic"), Context, OwnedObjects), OwnedObjects);
+		if (!TestInstance)
 		{
-			return true;
+			return 0u;
 		}
 
-		CDOGenerated = GeneratedTestInstance->GetClass()->ClassDefaultObject;
-		CDONative = NativeTestInstance->GetClass()->ClassDefaultObject;
-	}
+		FArchiveSkipTransientObjectCRC32 Results;
+		return Results.Crc32(TestInstance->GetClass()->ClassDefaultObject);
+	};
 
-	// evaluate results:
-	if (!Identical(CDOGenerated, CDONative))
-	{
-		AddError(FString::Printf(TEXT("%S: Generated CDO differs from native CDO"), *CDOGenerated->GetName()));
-	}
-	return true;
+	return RunTestHelper(TestBody, this);
 }
 
 IMPLEMENT_SIMPLE_AUTOMATION_TEST(FBPCompilerCommunicationTest, "Project.Blueprints.NativeBackend.CommunicationTest", CompilerTestFlags)
 bool FBPCompilerCommunicationTest::RunTest(const FString& Parameters)
 {
-	// setup input data:
-	UObject* GeneratedTestInstanceA = NewTestObject(GetGeneratedClass(TEXT("Communication"), TEXT("BP_Comm_Test_A"), this));
-	UObject* GeneratedTestInstanceB = NewTestObject(GetGeneratedClass(TEXT("Communication"), TEXT("BP_Comm_Test_B"), this));
-	if (!GeneratedTestInstanceA || !GeneratedTestInstanceB)
+	auto TestBody = [](ClassAccessor F, FAutomationTestBase* Context)
 	{
-		return true;
-	}
+		FOwnedObjectsHelper OwnedObjects;
 
-	UObject* NativeTestInstanceA = NewTestObject(GetNativeClass(TEXT("BP_Comm_Test_A_C"), this));
-	UObject* NativeTestInstanceB = NewTestObject(GetNativeClass(TEXT("BP_Comm_Test_B_C"), this));
+		UObject* A = NewTestObject(F(TEXT("Communication"), TEXT("BP_Comm_Test_A"), Context, OwnedObjects), OwnedObjects);
+		UObject* B = NewTestObject(F(TEXT("Communication"), TEXT("BP_Comm_Test_B"), Context, OwnedObjects), OwnedObjects);
+		if (!A || !B)
+		{
+			return 0u;
+		}
 
-	// execute test on generated object and native object:
-	auto Test = [](UObject* A, UObject* B)
-	{
 		Call(A, TEXT("Flop"));
 		Call(B, TEXT("Flip"));
+
+		FArchiveSkipTransientObjectCRC32 Results;
+		return Results.Crc32(B, Results.Crc32(A));
 	};
-	Test(GeneratedTestInstanceA, GeneratedTestInstanceB);
-	Test(NativeTestInstanceA, NativeTestInstanceB);
 
-	// evaluate results:
-	if (!Identical(GeneratedTestInstanceA, NativeTestInstanceA))
-	{
-		AddError(TEXT("Different results for communication test object A"));
-	}
-	if (!Identical(GeneratedTestInstanceB, NativeTestInstanceB))
-	{
-		AddError(TEXT("Different results for communication test object B"));
-	}
-	
-	return true;
+	return RunTestHelper(TestBody, this);
 }
 
 IMPLEMENT_SIMPLE_AUTOMATION_TEST(FBPCompilerConstructionScriptTest, "Project.Blueprints.NativeBackend.ConstructionScriptTest", CompilerTestFlags)
 bool FBPCompilerConstructionScriptTest::RunTest(const FString& Parameters)
 {
-	// setup input data:
-	UObject* GeneratedTestInstance = NewTestActor(GetGeneratedClass(TEXT("ConstructionScript"), TEXT("BP_ConstructionScript_Test"), this));
-	UObject* NativeTestInstance = NewTestActor(GetNativeClass(TEXT("BP_ConstructionScript_Test_C"), this));
-
-	// evaluate results:
-	if (!Identical(GeneratedTestInstance, NativeTestInstance))
+	auto TestBody = [](ClassAccessor F, FAutomationTestBase* Context)
 	{
-		AddError(FString::Printf(TEXT("%S: Generated instance differs from native instance"), *GeneratedTestInstance->GetName()));
-	}
-	return true;
+		FOwnedObjectsHelper OwnedObjects;
+
+		UObject* TestInstance = NewTestActor(F(TEXT("ConstructionScript"), TEXT("BP_ConstructionScript_Test"), Context, OwnedObjects), OwnedObjects);
+		if (!TestInstance)
+		{
+			return 0u;
+		}
+
+		FArchiveSkipTransientObjectCRC32 Results;
+		return Results.Crc32(TestInstance);
+	};
+
+	return RunTestHelper(TestBody, this);
 }
 
 IMPLEMENT_SIMPLE_AUTOMATION_TEST(FBPCompilerControlFlowTest, "Project.Blueprints.NativeBackend.ControlFlow", CompilerTestFlags)
 bool FBPCompilerControlFlowTest::RunTest(const FString& Parameters)
 {
-	// setup input data:
-	UObject* GeneratedTestInstance = NewTestObject(GetGeneratedClass(TEXT("ControlFlow"), TEXT("BP_ControlFlow_Basic"), this));
-	UObject* NativeTestInstance = NewTestObject(GetNativeClass(TEXT("BP_ControlFlow_Basic_C"), this));
+	auto TestBody = [](ClassAccessor F, FAutomationTestBase* Context)
+	{
+		FOwnedObjectsHelper OwnedObjects;
 
-	// execute test on generated object and native object:
-	Call(GeneratedTestInstance, TEXT("RunControlFlowTest"));
-	Call(NativeTestInstance, TEXT("RunControlFlowTest"));
+		UObject* TestInstance = NewTestObject(F(TEXT("ControlFlow"), TEXT("BP_ControlFlow_Basic"), Context, OwnedObjects), OwnedObjects);
+		if (!TestInstance)
+		{
+			return 0u;
+		}
 
-	// evaluate results:
-	if (!Identical(GeneratedTestInstance, NativeTestInstance))
-	{
-		AddError(FString::Printf(TEXT("%S: Generated instance differs from native instance"), *GeneratedTestInstance->GetName()));
-	}
-	return true;
+		Call(TestInstance, TEXT("RunControlFlowTest"));
+
+		FArchiveSkipTransientObjectCRC32 Results;
+		return Results.Crc32(TestInstance);
+	};
+
+	return RunTestHelper(TestBody, this);
 }
 
 IMPLEMENT_SIMPLE_AUTOMATION_TEST(FBPCompilerEnumTest, "Project.Blueprints.NativeBackend.EnumTest", CompilerTestFlags)
 bool FBPCompilerEnumTest::RunTest(const FString& Parameters)
 {
-	// setup input data:
-	UObject* GeneratedTestInstance = NewTestObject(GetGeneratedClass(TEXT("Enum"), TEXT("BP_Enum_Reader_Writer"), this));
-	UObject* NativeTestInstance = NewTestObject(GetNativeClass(TEXT("BP_Enum_Reader_Writer_C"), this));
+	auto TestBody = [](ClassAccessor F, FAutomationTestBase* Context)
+	{
+		FOwnedObjectsHelper OwnedObjects;
 
-	// execute test on generated object and native object:
-	Call(GeneratedTestInstance, TEXT("UpdateEnum"));
-	Call(NativeTestInstance, TEXT("UpdateEnum"));
+		UObject* TestInstance = NewTestObject(F(TEXT("Enum"), TEXT("BP_Enum_Reader_Writer"), Context, OwnedObjects), OwnedObjects);
+		if (!TestInstance)
+		{
+			return 0u;
+		}
 
-	// evaluate results:
-	if (!Identical(GeneratedTestInstance, NativeTestInstance))
-	{
-		AddError(FString::Printf(TEXT("%S: Generated instance differs from native instance"), *GeneratedTestInstance->GetName()));
-	}
-	return true;
+		Call(TestInstance, TEXT("UpdateEnum"));
+
+		FArchiveSkipTransientObjectCRC32 Results;
+		return Results.Crc32(TestInstance);
+	};
+
+	return RunTestHelper(TestBody, this);
 }
 
 IMPLEMENT_SIMPLE_AUTOMATION_TEST(FBPCompilerEventTest, "Project.Blueprints.NativeBackend.Event", CompilerTestFlags)
 bool FBPCompilerEventTest::RunTest(const FString& Parameters)
 {
-	// setup input data:
-	UObject* GeneratedTestInstance = NewTestObject(GetGeneratedClass(TEXT("Event"), TEXT("BP_Event_Basic"), this));
-	UObject* NativeTestInstance = NewTestObject(GetNativeClass(TEXT("BP_Event_Basic_C"), this));
+	auto TestBody = [](ClassAccessor F, FAutomationTestBase* Context)
+	{
+		FOwnedObjectsHelper OwnedObjects;
 
-	// execute test on generated object and native object:
-	Call(GeneratedTestInstance, TEXT("BeginEventChain"));
-	Call(NativeTestInstance, TEXT("BeginEventChain"));
+		UObject* TestInstance = NewTestObject(F(TEXT("Event"), TEXT("BP_Event_Basic"), Context, OwnedObjects), OwnedObjects);
+		if (!TestInstance)
+		{
+			return 0u;
+		}
 
-	// evaluate results:
-	if (!Identical(GeneratedTestInstance, NativeTestInstance))
-	{
-		AddError(FString::Printf(TEXT("%S: Generated instance differs from native instance"), *GeneratedTestInstance->GetName()));
-	}
-	return true;
+		Call(TestInstance, TEXT("BeginEventChain"));
+
+		FArchiveSkipTransientObjectCRC32 Results;
+		return Results.Crc32(TestInstance);
+	};
+
+	return RunTestHelper(TestBody, this);
 }
 
 struct FInheritenceTestParams
@@ -333,7 +350,7 @@ struct FInheritenceTestParams
 
 inline bool operator==(const FInheritenceTestParams& LHS, const FInheritenceTestParams& RHS)
 {
-	return 
+	return
 		LHS.bFlag == RHS.bFlag &&
 		LHS.Strings == RHS.Strings &&
 		LHS.Result == RHS.Result;
@@ -347,63 +364,91 @@ inline bool operator!=(const FInheritenceTestParams& LHS, const FInheritenceTest
 IMPLEMENT_SIMPLE_AUTOMATION_TEST(FBPCompilerInheritenceTest, "Project.Blueprints.NativeBackend.Inheritence", CompilerTestFlags)
 bool FBPCompilerInheritenceTest::RunTest(const FString& Parameters)
 {
-	// setup input data:
-	UObject* GeneratedTestInstance = NewTestObject(GetGeneratedClass(TEXT("Inheritence"), TEXT("BP_Child_Basic"), this));
-	UObject* NativeTestInstance = NewTestObject(GetNativeClass(TEXT("BP_Child_Basic_C"), this));
-
-	// execute test on generated object and native object:
-	FInheritenceTestParams ParamsGenerated;
-	ParamsGenerated.bFlag = true;
-	FInheritenceTestParams ParamsNative = ParamsGenerated;
-	Call(GeneratedTestInstance, TEXT("VirtualFunction"), &ParamsGenerated);
-	Call(NativeTestInstance, TEXT("VirtualFunction"), &ParamsNative);
-
-	// evaluate results:
-	if (ParamsGenerated != ParamsNative)
+	auto TestBody = [](ClassAccessor F, FAutomationTestBase* Context)
 	{
-		AddError(FString::Printf(TEXT("%S: Outputs differ for native and generated instance"), *GeneratedTestInstance->GetName()));
-	}
-	if (!Identical(GeneratedTestInstance, NativeTestInstance))
-	{
-		AddError(FString::Printf(TEXT("%S: Generated instance differs from native instance"), *GeneratedTestInstance->GetName()));
-	}
-	return true;
+		FOwnedObjectsHelper OwnedObjects;
+
+		UObject* TestInstance = NewTestObject(F(TEXT("Inheritence"), TEXT("BP_Child_Basic"), Context, OwnedObjects), OwnedObjects);
+		if (!TestInstance)
+		{
+			return 0u;
+		}
+
+		FInheritenceTestParams Params;
+		Params.bFlag = true;
+		Call(TestInstance, TEXT("VirtualFunction"), &Params);
+
+		FArchiveSkipTransientObjectCRC32 Results;
+		return Results.Crc32(TestInstance);
+	};
+
+	return RunTestHelper(TestBody, this);
 }
 
 IMPLEMENT_SIMPLE_AUTOMATION_TEST(FBPCompilerStructureTest, "Project.Blueprints.NativeBackend.Structure", CompilerTestFlags)
 bool FBPCompilerStructureTest::RunTest(const FString& Parameters)
 {
-	// setup input data:
-	UObject* GeneratedTestInstance = NewTestObject(GetGeneratedClass(TEXT("Structure"), TEXT("BP_Structure_Driver"), this));
-	UObject* NativeTestInstance = NewTestObject(GetNativeClass(TEXT("BP_Structure_Driver_C"), this));
+	auto TestBody = [](ClassAccessor F, FAutomationTestBase* Context)
+	{
+		FOwnedObjectsHelper OwnedObjects;
 
-	// execute test on generated object and native object:
-	Call(GeneratedTestInstance, TEXT("RunStructTest"));
-	Call(NativeTestInstance, TEXT("RunStructTest"));
+		UObject* TestInstance = NewTestObject(F(TEXT("Structure"), TEXT("BP_Structure_Driver"), Context, OwnedObjects), OwnedObjects);
+		if (!TestInstance)
+		{
+			return 0u;
+		}
 
-	// evaluate results:
-	if (!Identical(GeneratedTestInstance, NativeTestInstance))
-	{
-		AddError(FString::Printf(TEXT("%S: Generated instance differs from native instance"), *GeneratedTestInstance->GetName()));
-	}
-	return true;
+		Call(TestInstance, TEXT("RunStructTest"));
+
+		FArchiveSkipTransientObjectCRC32 Results;
+		return Results.Crc32(TestInstance);
+	};
+
+	return RunTestHelper(TestBody, this);
 }
 
 IMPLEMENT_SIMPLE_AUTOMATION_TEST(FBPCompilerNodeTest, "Project.Blueprints.NativeBackend.Node", CompilerTestFlags)
 bool FBPCompilerNodeTest::RunTest(const FString& Parameters)
 {
-	// setup input data:
-	UObject* GeneratedTestInstance = NewTestObject(GetGeneratedClass(TEXT("Node"), TEXT("BP_Node_Basic"), this));
-	UObject* NativeTestInstance = NewTestObject(GetNativeClass(TEXT("BP_Node_Basic_C"), this));
+	auto TestBody = [](ClassAccessor F, FAutomationTestBase* Context)
+	{
+		FOwnedObjectsHelper OwnedObjects;
+
+		UObject* TestInstance = NewTestObject(F(TEXT("Node"), TEXT("BP_Node_Basic"), Context, OwnedObjects), OwnedObjects);
+		if (!TestInstance)
+		{
+			return 0u;
+		}
 
-	// execute test on generated object and native object:
-	Call(GeneratedTestInstance, TEXT("RunNodes"));
-	Call(NativeTestInstance, TEXT("RunNodes"));
+		Call(TestInstance, TEXT("RunNodes"));
 
-	// evaluate results:
-	if (!Identical(GeneratedTestInstance, NativeTestInstance))
+		FArchiveSkipTransientObjectCRC32 Results;
+		return Results.Crc32(TestInstance);
+	};
+
+	return RunTestHelper(TestBody, this);
+}
+
+IMPLEMENT_SIMPLE_AUTOMATION_TEST(FBPCompilerLatentTest, "Project.Blueprints.NativeBackend.Latent", CompilerTestFlags)
+bool FBPCompilerLatentTest::RunTest(const FString& Parameters)
+{
+	auto TestBody = [](ClassAccessor F, FAutomationTestBase* Context)
 	{
-		AddError(FString::Printf(TEXT("%S: Generated instance differs from native instance"), *GeneratedTestInstance->GetName()));
-	}
-	return true;
+		FOwnedObjectsHelper OwnedObjects;
+		TGuardValue<bool> AutoRestore(GAllowActorScriptExecutionInEditor, true);
+
+		UObject* TestInstance = NewTestActor(F(TEXT("Node"), TEXT("BP_Latent_Basic"), Context, OwnedObjects), OwnedObjects);
+		if (!TestInstance)
+		{
+			return 0u;
+		}
+
+		Call(TestInstance, TEXT("RunDelayTest"));
+		Call(TestInstance, TEXT("RunDownloadTest"));
+
+		FArchiveSkipTransientObjectCRC32 Results;
+		return Results.Crc32(TestInstance);
+	};
+
+	return RunTestHelper(TestBody, this);
 }
diff --git a/Engine/Source/Developer/BlueprintCompilerCppBackend/Private/BlueprintCompilerCppBackend.cpp b/Engine/Source/Developer/BlueprintCompilerCppBackend/Private/BlueprintCompilerCppBackend.cpp
--- a/Engine/Source/Developer/BlueprintCompilerCppBackend/Private/BlueprintCompilerCppBackend.cpp
+++ b/Engine/Source/Developer/BlueprintCompilerCppBackend/Private/BlueprintCompilerCppBackend.cpp
@@ -184,7 +184,7 @@ void FBlueprintCompilerCppBackend::EmitObjectToBoolStatement(FEmitterLocalContex
 {
 	FString ObjectTarget = TermToText(EmitterContext, Statement.RHS[0]);
 	FString DestinationExpression = TermToText(EmitterContext, Statement.LHS);
-	EmitterContext.AddLine(FString::Printf(TEXT("%s = (nullptr != %s);"), *DestinationExpression, *ObjectTarget));
+	EmitterContext.AddLine(FString::Printf(TEXT("%s = (%s != nullptr);"), *DestinationExpression, *ObjectTarget));
 }
 
 void FBlueprintCompilerCppBackend::EmitAddMulticastDelegateStatement(FEmitterLocalContext& EmitterContext, FKismetFunctionContext& FunctionContext, FBlueprintCompiledStatement& Statement)
@@ -316,11 +316,6 @@ FString FBlueprintCompilerCppBackend::EmitSwitchValueStatmentInner(FEmitterLocal
 	auto IndexTerm = Statement.RHS[0];
 	auto DefaultValueTerm = Statement.RHS.Last();
 
-	FString Result = FString::Printf(TEXT("TSwitchValue(%s, %s, %d")
-		, *TermToText(EmitterContext, IndexTerm) //index
-		, *TermToText(EmitterContext, DefaultValueTerm) // default
-		, NumCases);
-
 	const uint32 CppTemplateTypeFlags = EPropertyExportCPPFlags::CPPF_CustomTypeName
 		| EPropertyExportCPPFlags::CPPF_NoConst | EPropertyExportCPPFlags::CPPF_NoRef
 		| EPropertyExportCPPFlags::CPPF_BlueprintCppBackend;
@@ -331,6 +326,13 @@ FString FBlueprintCompilerCppBackend::EmitSwitchValueStatmentInner(FEmitterLocal
 	check(DefaultValueTerm && DefaultValueTerm->AssociatedVarProperty);
 	const FString ValueDeclaration = EmitterContext.ExportCppDeclaration(DefaultValueTerm->AssociatedVarProperty, EExportedDeclaration::Parameter, CppTemplateTypeFlags, true);
 
+	FString Result = FString::Printf(TEXT("TSwitchValue<%s, %s>(%s, %s, %d")
+		, *IndexDeclaration
+		, *ValueDeclaration
+		, *TermToText(EmitterContext, IndexTerm) //index
+		, *TermToText(EmitterContext, DefaultValueTerm) // default
+		, NumCases);
+
 	for (int32 TermIndex = TermsBeforeCases; TermIndex < (NumCases * TermsPerCase); TermIndex += TermsPerCase)
 	{
 		Result += FString::Printf(TEXT(", TSwitchPair<%s, %s>(%s, %s)")
@@ -439,13 +441,14 @@ FString FBlueprintCompilerCppBackend::EmitCallStatmentInner(FEmitterLocalContext
 	}
 
 	// Emit object to call the method on
+	const FString FunctionToCallOriginalName = FEmitHelper::GetCppName(FEmitHelper::GetOriginalFunction(Statement.FunctionToCall));
 	if (bInterfaceCall)
 	{
 		auto ContextInterfaceClass = CastChecked<UClass>(Statement.FunctionContext->Type.PinSubCategoryObject.Get());
 		ensure(ContextInterfaceClass->IsChildOf<UInterface>());
-		Result += FString::Printf(TEXT("%s::Execute_%s(%s.GetObject(), ")
+		Result += FString::Printf(TEXT("%s::Execute_%s(%s.GetObject() ")
 			, *FEmitHelper::GetCppName(ContextInterfaceClass)
-			, *FEmitHelper::GetCppName(Statement.FunctionToCall)
+			, *FunctionToCallOriginalName
 			, *TermToText(EmitterContext, Statement.FunctionContext, false));
 	}
 	else
@@ -479,7 +482,7 @@ FString FBlueprintCompilerCppBackend::EmitCallStatmentInner(FEmitterLocalContext
 		{
 			Result += TEXT("Super::");
 		}
-		Result += FEmitHelper::GetCppName(Statement.FunctionToCall);
+		Result += FunctionToCallOriginalName;
 		if (Statement.bIsParentContext && FEmitHelper::ShouldHandleAsNativeEvent(Statement.FunctionToCall))
 		{
 			ensure(!bCallOnDifferentObject);
@@ -489,7 +492,12 @@ FString FBlueprintCompilerCppBackend::EmitCallStatmentInner(FEmitterLocalContext
 		// Emit method parameter list
 		Result += TEXT("(");
 	}
-	Result += EmitMethodInputParameterList(EmitterContext, Statement);
+	const FString ParameterList = EmitMethodInputParameterList(EmitterContext, Statement);
+	if (bInterfaceCall && !ParameterList.IsEmpty())
+	{
+		Result += TEXT(", ");
+	}
+	Result += ParameterList;
 	Result += TEXT(")");
 	Result += CloseCast;
 	if (!bInline)
@@ -556,10 +564,22 @@ FString FBlueprintCompilerCppBackend::TermToText(FEmitterLocalContext& EmitterCo
 		if (Term->Context && Term->Context->IsStructContextType())
 		{
 			check(Term->AssociatedVarProperty);
-			ResultPath = ContextStr + TEXT(".") + FEmitHelper::GetCppName(Term->AssociatedVarProperty);
+			const bool bIsAccessible = !Term->AssociatedVarProperty->HasAnyPropertyFlags(CPF_NativeAccessSpecifierPrivate | CPF_NativeAccessSpecifierProtected);
+			if (!bIsAccessible)
+			{
+				ResultPath = FEmitHelper::AccessInaccessibleProperty(EmitterContext, Term->AssociatedVarProperty, ContextStr, TEXT("&"));
+			}
+			else
+			{
+				ResultPath = ContextStr + TEXT(".") + FEmitHelper::GetCppName(Term->AssociatedVarProperty);
+			}
 		}
 		else if (Term->AssociatedVarProperty)
 		{
+			const bool bPropertyOfParent = EmitterContext.Dependencies.GetActualStruct()->IsChildOf(Term->AssociatedVarProperty->GetOwnerStruct());
+			const bool bIsAccessible = !Term->AssociatedVarProperty->HasAnyPropertyFlags(CPF_NativeAccessSpecifierPrivate)
+				&& (bPropertyOfParent || !Term->AssociatedVarProperty->HasAnyPropertyFlags(CPF_NativeAccessSpecifierProtected));
+
 			auto MinimalClass = Term->AssociatedVarProperty->GetOwnerClass();
 			auto MinimalBPGC = Cast<UBlueprintGeneratedClass>(MinimalClass);
 			if (MinimalBPGC && !EmitterContext.Dependencies.WillClassBeConverted(MinimalBPGC))
@@ -575,6 +595,15 @@ FString FBlueprintCompilerCppBackend::TermToText(FEmitterLocalContext& EmitterCo
 					, *ContextStr
 					, *UnicodeToCPPIdentifier(Term->AssociatedVarProperty->GetName(), false, nullptr));
 			}
+			else if (!bIsAccessible)
+			{
+				if ((!Term->Context) || (Term->Context->Name == PSC_Self))
+				{
+					ensure(ContextStr.IsEmpty());
+					ContextStr = TEXT("this");
+				}
+				ResultPath = FEmitHelper::AccessInaccessibleProperty(EmitterContext, Term->AssociatedVarProperty, ContextStr, FString());
+			}
 			else
 			{
 				if ((Term->Context != nullptr) && (Term->Context->Name != PSC_Self))
@@ -587,7 +616,7 @@ FString FBlueprintCompilerCppBackend::TermToText(FEmitterLocalContext& EmitterCo
 		else
 		{
 			ensure(ContextStr.IsEmpty());
-			ResultPath += Term->Name;
+			ResultPath = Term->Name;
 		}
 
 		FString Conditions;
@@ -596,9 +625,20 @@ FString FBlueprintCompilerCppBackend::TermToText(FEmitterLocalContext& EmitterCo
 			Conditions = FSafeContextScopedEmmitter::ValidationChain(EmitterContext, Term->Context, *this);
 		}
 
-		return Conditions.IsEmpty()
-			? ResultPath
-			: FString::Printf(TEXT("((%s) ? (%s) : (%s))"), *Conditions, *ResultPath, *FEmitHelper::DefaultValue(EmitterContext, Term->Type));
+		if (!Conditions.IsEmpty())
+		{
+			FString DefaultValue = FEmitHelper::DefaultValue(EmitterContext, Term->Type);
+			if (Term->Type.bIsArray)
+			{
+				// to fix error C4239 in VS, when the term is passed as non-const reference. 
+				// It seems that this situation happens only for arrays, otherwise this should be a default path.
+				const FString DefaultValueVariable = EmitterContext.GenerateUniqueLocalName();
+				EmitterContext.AddLine(*FString::Printf(TEXT("auto %s = %s;"), *DefaultValueVariable, *DefaultValue));
+				DefaultValue = DefaultValueVariable;
+			}
+			return FString::Printf(TEXT("((%s) ? (%s) : (%s))"), *Conditions, *ResultPath, *DefaultValue);
+		}
+		return ResultPath;
 	}
 }
 
diff --git a/Engine/Source/Developer/BlueprintCompilerCppBackend/Private/BlueprintCompilerCppBackendAnim.cpp b/Engine/Source/Developer/BlueprintCompilerCppBackend/Private/BlueprintCompilerCppBackendAnim.cpp
new file mode 100644
--- /dev/null
+++ b/Engine/Source/Developer/BlueprintCompilerCppBackend/Private/BlueprintCompilerCppBackendAnim.cpp
@@ -0,0 +1,36 @@
+// Copyright 1998-2015 Epic Games, Inc. All Rights Reserved.
+
+#include "BlueprintCompilerCppBackendModulePrivatePCH.h"
+#include "BlueprintCompilerCppBackendUtils.h"
+#include "Animation/AnimBlueprintGeneratedClass.h"
+#include "Animation/AnimClassData.h"
+
+FString FBackendHelperAnim::AddHeaders(UClass* GeneratedClass)
+{
+	auto AnimClass = Cast<UAnimBlueprintGeneratedClass>(GeneratedClass);
+	return AnimClass ? FString(TEXT("#include \"Animation/AnimClassData.h\"")) : FString();
+}
+
+void FBackendHelperAnim::CreateAnimClassData(FEmitterLocalContext& Context)
+{
+	if (auto AnimClass = Cast<UAnimBlueprintGeneratedClass>(Context.GetCurrentlyGeneratedClass()))
+	{
+		const FString LocalNativeName = Context.GenerateUniqueLocalName();
+		Context.AddLine(FString::Printf(TEXT("auto %s = NewObject<UAnimClassData>(GetClass(), TEXT(\"AnimClassData\"));"), *LocalNativeName));
+
+		auto AnimClassData = NewObject<UAnimClassData>(GetTransientPackage(), TEXT("AnimClassData"));
+		AnimClassData->CopyFrom(AnimClass);
+
+		auto ObjectArchetype = AnimClassData->GetArchetype();
+		for (auto Property : TFieldRange<const UProperty>(UAnimClassData::StaticClass()))
+		{
+			FEmitDefaultValueHelper::OuterGenerate(Context, Property, LocalNativeName
+				, reinterpret_cast<const uint8*>(AnimClassData)
+				, reinterpret_cast<const uint8*>(ObjectArchetype)
+				, FEmitDefaultValueHelper::EPropertyAccessOperator::Pointer);
+		}
+
+		Context.AddLine(FString::Printf(TEXT("CastChecked<UDynamicClass>(GetClass())->AnimClassImplementation = %s;"), *LocalNativeName));
+	}
+}
+
diff --git a/Engine/Source/Developer/BlueprintCompilerCppBackend/Private/BlueprintCompilerCppBackendBase.cpp b/Engine/Source/Developer/BlueprintCompilerCppBackend/Private/BlueprintCompilerCppBackendBase.cpp
--- a/Engine/Source/Developer/BlueprintCompilerCppBackend/Private/BlueprintCompilerCppBackendBase.cpp
+++ b/Engine/Source/Developer/BlueprintCompilerCppBackend/Private/BlueprintCompilerCppBackendBase.cpp
@@ -53,6 +53,27 @@ void FBlueprintCompilerCppBackendBase::DeclareDelegates(FEmitterLocalContext& Em
 			}
 		}
 
+		TArray<UFunction*> SCDelegateSignaturesWithoutType;
+		// Don't declare signatures, that are already declared in a native class
+		for (int32 I = 0; I < Delegates.Num();)
+		{
+			auto Delegate = Delegates[I];
+			auto DelegateSignature = Delegate ? Delegate->SignatureFunction : nullptr;
+			auto DelegateSignatureOwner = DelegateSignature ? DelegateSignature->GetOwnerStruct() : nullptr;
+			if (DelegateSignatureOwner && DelegateSignatureOwner->HasAnyInternalFlags(EInternalObjectFlags::Native))
+			{
+				if (DelegateSignature->HasAllFunctionFlags(FUNC_MulticastDelegate))
+				{
+					SCDelegateSignaturesWithoutType.AddUnique(DelegateSignature);
+				}
+				Delegates.RemoveAtSwap(I);
+			}
+			else
+			{
+				I++;
+			}
+		}
+
 		// remove duplicates, n^2, but n is small:
 		for (int32 I = 0; I < Delegates.Num(); ++I)
 		{
@@ -68,6 +89,13 @@ void FBlueprintCompilerCppBackendBase::DeclareDelegates(FEmitterLocalContext& Em
 			}
 		}
 
+		for (UFunction* SCDelegateSignature : SCDelegateSignaturesWithoutType)
+		{
+			FString SCType = FString::Printf(TEXT("F__%s__SC"), *FEmitHelper::GetCppName(SCDelegateSignature));
+			FEmitHelper::EmitSinglecastDelegateDeclarations_Inner(EmitterContext, SCDelegateSignature, SCType);
+			EmitterContext.MCDelegateSignatureToSCDelegateType.Add(SCDelegateSignature, SCType);
+		}
+
 		FEmitHelper::EmitSinglecastDelegateDeclarations(EmitterContext, Delegates);
 	}
 }
@@ -84,11 +112,14 @@ FString FBlueprintCompilerCppBackendBase::GenerateCodeFromClass(UClass* SourceCl
 	FEmitterLocalContext EmitterContext(Dependencies);
 	EmitFileBeginning(CleanCppClassName, EmitterContext);
 
+	UClass* OriginalSourceClass = Dependencies.FindOriginalClass(SourceClass);
+	ensure(OriginalSourceClass != SourceClass);
+
 	// Class declaration
 	const bool bIsInterface = SourceClass->IsChildOf<UInterface>();
 	if (bIsInterface)
 	{
-		EmitterContext.Header.AddLine(FString::Printf(TEXT("UINTERFACE(Blueprintable %s)"), *FEmitHelper::ReplaceConvertedMetaData(SourceClass)));
+		EmitterContext.Header.AddLine(FString::Printf(TEXT("UINTERFACE(Blueprintable, %s)"), *FEmitHelper::ReplaceConvertedMetaData(OriginalSourceClass)));
 		EmitterContext.Header.AddLine(FString::Printf(TEXT("class %s : public UInterface"), *FEmitHelper::GetCppName(SourceClass, true)));
 		EmitterContext.Header.AddLine(TEXT("{"));
 		EmitterContext.Header.IncreaseIndent();
@@ -100,8 +131,8 @@ FString FBlueprintCompilerCppBackendBase::GenerateCodeFromClass(UClass* SourceCl
 	else
 	{
 		EmitterContext.Header.AddLine(FString::Printf(TEXT("UCLASS(%s%s)")
-			, (!SourceClass->IsChildOf<UBlueprintFunctionLibrary>()) ? TEXT("Blueprintable, BlueprintType") : TEXT("")
-			, *FEmitHelper::ReplaceConvertedMetaData(SourceClass)));
+			, (!SourceClass->IsChildOf<UBlueprintFunctionLibrary>()) ? TEXT("Blueprintable, BlueprintType, ") : TEXT("")
+			, *FEmitHelper::ReplaceConvertedMetaData(OriginalSourceClass)));
 
 		UClass* SuperClass = SourceClass->GetSuperClass();
 		FString ClassDefinition = FString::Printf(TEXT("class %s : public %s"), *CppClassName, *FEmitHelper::GetCppName(SuperClass));
@@ -131,7 +162,7 @@ FString FBlueprintCompilerCppBackendBase::GenerateCodeFromClass(UClass* SourceCl
 	{
 		EmitterContext.Header.AddLine(FString::Printf(TEXT("%s(const FObjectInitializer& ObjectInitializer = FObjectInitializer::Get());"), *CppClassName));
 		EmitterContext.Header.AddLine(TEXT("virtual void PostLoadSubobjects(FObjectInstancingGraph* OuterInstanceGraph) override;"));
-		EmitterContext.Header.AddLine(TEXT("static void __StaticDependenciesAssets(TArray<FName>& OutPackagePaths);"));
+		EmitterContext.Header.AddLine(TEXT("static void __StaticDependenciesAssets(TArray<FBlueprintDependencyData>& AssetsToLoad);"));
 
 		FEmitDefaultValueHelper::GenerateConstructor(EmitterContext);
 	}
@@ -186,173 +217,184 @@ void FBlueprintCompilerCppBackendBase::ConstructFunction(FKismetFunctionContext&
 		return;
 	}
 
-	UFunction* Function = FunctionContext.Function;
-
-	UProperty* ReturnValue = NULL;
 	TArray<UProperty*> LocalVariables;
+	TArray<UProperty*> ArgumentList;
+	// Split the function property list into arguments, a return value (if any), and local variable declarations
+	for (UProperty* Property : TFieldRange<UProperty>(FunctionContext.Function))
 	{
-		FString FunctionName = FEmitHelper::GetCppName(Function);
+		const bool bNeedLocalVariable = !Property->HasAnyPropertyFlags(CPF_Parm) || Property->HasAnyPropertyFlags(CPF_ReturnParm);
+		TArray<UProperty*>& PropertyDest = bNeedLocalVariable ? LocalVariables : ArgumentList;
+		PropertyDest.Add(Property);
+	}
 
-		TArray<UProperty*> ArgumentList;
+	if (FunctionContext.bIsUbergraph)
+	{
+		EmitterContext.AddLine(TEXT("PRAGMA_DISABLE_OPTIMIZATION"));
+	}
 
-		// Split the function property list into arguments, a return value (if any), and local variable declarations
-		for (TFieldIterator<UProperty> It(Function); It; ++It)
-		{
-			UProperty* Property = *It;
-			if (Property->HasAnyPropertyFlags(CPF_Parm))
-			{
-				if (Property->HasAnyPropertyFlags(CPF_ReturnParm))
-				{
-					if (ReturnValue == NULL)
-					{
-						ReturnValue = Property;
-						LocalVariables.Add(Property);
-					}
-					else
-					{
-						UE_LOG(LogK2Compiler, Error, TEXT("Function %s from graph @@ has more than one return value (named %s and %s)"),
-							*FunctionName, *GetPathNameSafe(FunctionContext.SourceGraph), *ReturnValue->GetName(), *Property->GetName());
-					}
-				}
-				else
-				{
-					ArgumentList.Add(Property);
-				}
-			}
-			else
-			{
-				LocalVariables.Add(Property);
-			}
-		}
+	ConstructFunctionDeclaration(EmitterContext, FunctionContext, ArgumentList);
+
+	// Start the body of the implementation
+	EmitterContext.AddLine(TEXT("{"));
+	EmitterContext.IncreaseIndent();
+	if (!bGenerateStubOnly)
+	{
+		DeclareLocalVariables(EmitterContext, LocalVariables);
+		ConstructFunctionBody(EmitterContext, FunctionContext);
+	}
+
+	UProperty* ReturnValue = FunctionContext.Function->GetReturnProperty();
+	const FString ReturnValueString = ReturnValue ? (FString(TEXT(" ")) + FEmitHelper::GetCppName(ReturnValue)) : TEXT("");
+	EmitterContext.AddLine(FString::Printf(TEXT("return%s;"), *ReturnValueString));
+	EmitterContext.DecreaseIndent();
+	EmitterContext.AddLine(TEXT("}"));
 
-		bool bAddCppFromBpEventMD = false;
+	if (FunctionContext.bIsUbergraph)
+	{
+		EmitterContext.AddLine(TEXT("PRAGMA_ENABLE_OPTIMIZATION"));
+	}
+}
+
+void FBlueprintCompilerCppBackendBase::ConstructFunctionDeclaration(FEmitterLocalContext &EmitterContext, FKismetFunctionContext &FunctionContext, TArray<UProperty*> &ArgumentList)
+{
+	FString FunctionHeaderName = FEmitHelper::GetCppName(FunctionContext.Function);
+	FString FunctionBodyName = FunctionHeaderName;
+	const bool bStaticFunction = FunctionContext.Function->HasAllFunctionFlags(FUNC_Static);
+	const bool bInInterface = FunctionContext.Function->GetOwnerClass()->IsChildOf<UInterface>();
+	bool bAddConst = false;
+	bool bIsOverride = false;
+	bool bIsVirtual = !bInInterface && !bStaticFunction && !FunctionContext.IsEventGraph();
+
+	{
+		UFunction* const Function = FunctionContext.Function;
+		const UFunction* const OriginalFunction = FEmitHelper::GetOriginalFunction(Function);
+		TArray<FString> AdditionalMetaData;
+		TArray<FString> AdditionalTags;
 		bool bGenerateAsNativeEventImplementation = false;
-		// Get original function declaration
-		if (FEmitHelper::ShouldHandleAsNativeEvent(Function)) // BlueprintNativeEvent
+		const bool bNetImplementation = !bInInterface && Function->HasAllFunctionFlags(FUNC_Net) && !Function->HasAnyFunctionFlags(FUNC_NetResponse);
+
+		const UBlueprintGeneratedClass* const OriginalFuncOwnerAsBPGC = Cast<UBlueprintGeneratedClass>(OriginalFunction->GetOwnerClass());
+		const bool bBPInterfaceImplementation = OriginalFuncOwnerAsBPGC && OriginalFuncOwnerAsBPGC->IsChildOf<UInterface>();
+
+		if (bInInterface)
 		{
-			bGenerateAsNativeEventImplementation = true;
-			FunctionName += TEXT("_Implementation");
+			AdditionalTags.Emplace(TEXT("BlueprintImplementableEvent"));
 		}
-		else if (FEmitHelper::ShouldHandleAsImplementableEvent(Function)) // BlueprintImplementableEvent
+		else if (bNetImplementation)
 		{
-			bAddCppFromBpEventMD = true;
+			FunctionBodyName = FunctionHeaderName + TEXT("_Implementation");
 		}
-
-		// Emit the declaration
-		FString ReturnType;
-		if (ReturnValue)
+		else if (FEmitHelper::ShouldHandleAsNativeEvent(Function))
 		{
-			const uint32 LocalExportCPPFlags = EPropertyExportCPPFlags::CPPF_CustomTypeName
-				| EPropertyExportCPPFlags::CPPF_NoConst 
-				| EPropertyExportCPPFlags::CPPF_NoRef 
-				| EPropertyExportCPPFlags::CPPF_NoStaticArray 
-				| EPropertyExportCPPFlags::CPPF_BlueprintCppBackend;
-			ReturnType = EmitterContext.ExportCppDeclaration(ReturnValue, EExportedDeclaration::Parameter, LocalExportCPPFlags, true);
+			bGenerateAsNativeEventImplementation = true;
+			FunctionBodyName = FunctionHeaderName = FEmitHelper::GetCppName(OriginalFunction) + TEXT("_Implementation");
+			bAddConst = OriginalFunction->HasAllFunctionFlags(FUNC_Const);
 		}
-		else
+		else if (FEmitHelper::ShouldHandleAsImplementableEvent(Function) || bBPInterfaceImplementation)
 		{
-			ReturnType = TEXT("void");
+			//The function "bpf__BIE__pf" should never be called directly. Only via function "BIE" with generated implementation.
+			bIsVirtual = false;
+			AdditionalMetaData.Emplace(TEXT("CppFromBpEvent"));
 		}
 
-		//@TODO: Make the header+body export more uniform
-		{
-			const FString Start = FString::Printf(TEXT("%s %s%s%s"), *ReturnType, TEXT("%s"), TEXT("%s"), *FunctionName);
-
-			TArray<FString> AdditionalMetaData;
-			if (bAddCppFromBpEventMD)
-			{
-				AdditionalMetaData.Emplace(TEXT("CppFromBpEvent"));
-			}
+		ensure(!bIsVirtual || Function->IsSignatureCompatibleWith(OriginalFunction));
+		bIsOverride = bGenerateAsNativeEventImplementation || (bIsVirtual && (Function != OriginalFunction));
 
-			if (!bGenerateAsNativeEventImplementation)
-			{
-				EmitterContext.Header.AddLine(FString::Printf(TEXT("%s"), *FEmitHelper::EmitUFuntion(Function, &AdditionalMetaData)));
-			}
-
-			FString HeaderFunction;
-			if (Function->HasAllFunctionFlags(FUNC_Static))
-			{
-				HeaderFunction += TEXT("static ");
-			}
-
-			if (bGenerateAsNativeEventImplementation)
-			{
-				HeaderFunction += TEXT("virtual ");
-			}
-
-			HeaderFunction += FString::Printf(*Start, TEXT(""), TEXT(""));
-			FString BodyFunction = FString::Printf(*Start, *FEmitHelper::GetCppName(EmitterContext.GetCurrentlyGeneratedClass()), TEXT("::"));
-
-			FString ArgListStr = TEXT("(");
-			for (int32 i = 0; i < ArgumentList.Num(); ++i)
-			{
-				UProperty* ArgProperty = ArgumentList[i];
+		if (bNetImplementation && bIsOverride)
+		{
+			FunctionHeaderName = FunctionBodyName;
+		}
+		else if (!bGenerateAsNativeEventImplementation)
+		{
+			EmitterContext.Header.AddLine(FEmitHelper::EmitUFuntion(Function, AdditionalTags, AdditionalMetaData));
+		}
+	}
 
-				if (i > 0)
-				{
-					ArgListStr += TEXT(", ");
-				}
+	const FString ReturnType = GenerateReturnType(EmitterContext, FunctionContext.Function);
+	const FString ArgList = GenerateArgList(EmitterContext, ArgumentList);
+
+	EmitterContext.Body.AddLine(FString::Printf(TEXT("%s %s::%s%s%s")
+		, *ReturnType
+		, *FEmitHelper::GetCppName(EmitterContext.GetCurrentlyGeneratedClass())
+		, *FunctionBodyName
+		, *ArgList
+		, bAddConst ? TEXT(" const") : TEXT("")));
+
+	EmitterContext.Header.AddLine(FString::Printf(TEXT("%s%s%s %s%s%s%s;")
+		, bStaticFunction ? TEXT("static ") : TEXT("")
+		, bIsVirtual ? TEXT("virtual ") : TEXT("")
+		, *ReturnType
+		, *FunctionHeaderName
+		, *ArgList
+		, bAddConst ? TEXT(" const") : TEXT("")
+		, bIsOverride ? TEXT(" override") : TEXT("")));
+}
 
-				if (ArgProperty->HasAnyPropertyFlags(CPF_OutParm))
-				{
-					ArgListStr += TEXT("/*out*/ ");
-				}
+FString FBlueprintCompilerCppBackendBase::GenerateArgList(const FEmitterLocalContext &EmitterContext, const TArray<UProperty*> &ArgumentList)
+{
+	FString ArgListStr = TEXT("(");
+	for (int32 i = 0; i < ArgumentList.Num(); ++i)
+	{
+		UProperty* ArgProperty = ArgumentList[i];
 
-				ArgListStr += EmitterContext.ExportCppDeclaration(ArgProperty, EExportedDeclaration::Parameter, EPropertyExportCPPFlags::CPPF_CustomTypeName | EPropertyExportCPPFlags::CPPF_BlueprintCppBackend);
-			}
-			ArgListStr += TEXT(")");
-			HeaderFunction += ArgListStr;
-			BodyFunction += ArgListStr;
-			if (bGenerateAsNativeEventImplementation)
-			{
-				HeaderFunction += TEXT(" override");
-			}
-			HeaderFunction += TEXT(";");
+		if (i > 0)
+		{
+			ArgListStr += TEXT(", ");
+		}
 
-			EmitterContext.AddLine(BodyFunction);
-			EmitterContext.Header.AddLine(HeaderFunction);
+		if (ArgProperty->HasAnyPropertyFlags(CPF_OutParm))
+		{
+			ArgListStr += TEXT("/*out*/ ");
 		}
+
+		ArgListStr += EmitterContext.ExportCppDeclaration(ArgProperty, EExportedDeclaration::Parameter, EPropertyExportCPPFlags::CPPF_CustomTypeName | EPropertyExportCPPFlags::CPPF_BlueprintCppBackend);
 	}
+	ArgListStr += TEXT(")");
+	return ArgListStr;
+}
 
-	// Start the body of the implementation
-	EmitterContext.AddLine(TEXT("{"));
-	EmitterContext.IncreaseIndent();
-	if (!bGenerateStubOnly)
+FString FBlueprintCompilerCppBackendBase::GenerateReturnType(const FEmitterLocalContext &EmitterContext, const UFunction* Function)
+{
+	UProperty* ReturnValue = Function->GetReturnProperty();
+	if (ReturnValue)
 	{
-		// Emit local variables
-		DeclareLocalVariables(EmitterContext, LocalVariables);
+		const uint32 LocalExportCPPFlags = EPropertyExportCPPFlags::CPPF_CustomTypeName
+			| EPropertyExportCPPFlags::CPPF_NoConst
+			| EPropertyExportCPPFlags::CPPF_NoRef
+			| EPropertyExportCPPFlags::CPPF_NoStaticArray
+			| EPropertyExportCPPFlags::CPPF_BlueprintCppBackend;
+		return EmitterContext.ExportCppDeclaration(ReturnValue, EExportedDeclaration::Parameter, LocalExportCPPFlags, true);
+	}
+	return TEXT("void");
+}
 
-		const bool bUseSwitchState = FunctionContext.MustUseSwitchState(nullptr);
+void FBlueprintCompilerCppBackendBase::ConstructFunctionBody(FEmitterLocalContext& EmitterContext, FKismetFunctionContext &FunctionContext)
+{
+	const bool bUseSwitchState = FunctionContext.MustUseSwitchState(nullptr);
 
-		// Run thru code looking only at things marked as jump targets, to make sure the jump targets are ordered in order of appearance in the linear execution list
-		// Emit code in the order specified by the linear execution list (the first node is always the entry point for the function)
-		for (int32 NodeIndex = 0; NodeIndex < FunctionContext.LinearExecutionList.Num(); ++NodeIndex)
-		{
-			UEdGraphNode* StatementNode = FunctionContext.LinearExecutionList[NodeIndex];
-			TArray<FBlueprintCompiledStatement*>* StatementList = FunctionContext.StatementsPerNode.Find(StatementNode);
+	// Run thru code looking only at things marked as jump targets, to make sure the jump targets are ordered in order of appearance in the linear execution list
+	// Emit code in the order specified by the linear execution list (the first node is always the entry point for the function)
+	for (int32 NodeIndex = 0; NodeIndex < FunctionContext.LinearExecutionList.Num(); ++NodeIndex)
+	{
+		UEdGraphNode* StatementNode = FunctionContext.LinearExecutionList[NodeIndex];
+		TArray<FBlueprintCompiledStatement*>* StatementList = FunctionContext.StatementsPerNode.Find(StatementNode);
 
-			if (StatementList != NULL)
+		if (StatementList != NULL)
+		{
+			for (int32 StatementIndex = 0; StatementIndex < StatementList->Num(); ++StatementIndex)
 			{
-				for (int32 StatementIndex = 0; StatementIndex < StatementList->Num(); ++StatementIndex)
-				{
-					FBlueprintCompiledStatement& Statement = *((*StatementList)[StatementIndex]);
+				FBlueprintCompiledStatement& Statement = *((*StatementList)[StatementIndex]);
 
-					if (Statement.bIsJumpTarget)
-					{
-						// Just making sure we number them in order of appearance, so jump statements don't influence the order
-						const int32 StateNum = StatementToStateIndex(FunctionContext, &Statement);
-					}
+				if (Statement.bIsJumpTarget)
+				{
+					// Just making sure we number them in order of appearance, so jump statements don't influence the order
+					const int32 StateNum = StatementToStateIndex(FunctionContext, &Statement);
 				}
 			}
 		}
-
-		InnerFunctionImplementation(FunctionContext, EmitterContext, bUseSwitchState);
 	}
 
-	const FString ReturnValueString = ReturnValue ? (FString(TEXT(" ")) + FEmitHelper::GetCppName(ReturnValue)) : TEXT("");
-	EmitterContext.AddLine(FString::Printf(TEXT("return%s;"), *ReturnValueString));
-	EmitterContext.DecreaseIndent();
-	EmitterContext.AddLine(TEXT("}"));
+	InnerFunctionImplementation(FunctionContext, EmitterContext, bUseSwitchState);
 }
 
 FString FBlueprintCompilerCppBackendBase::GenerateCodeFromEnum(UUserDefinedEnum* SourceEnum)
@@ -362,7 +404,7 @@ FString FBlueprintCompilerCppBackendBase::GenerateCodeFromEnum(UUserDefinedEnum*
 	Header.AddLine(TEXT("#pragma once"));
 	// use GetBaseFilename() so that we can coordinate #includes and filenames
 	Header.AddLine(FString::Printf(TEXT("#include \"%s.generated.h\""), *FEmitHelper::GetBaseFilename(SourceEnum)));
-	Header.AddLine(FString::Printf(TEXT("UENUM(BlueprintType %s )"), *FEmitHelper::ReplaceConvertedMetaData(SourceEnum)));
+	Header.AddLine(FString::Printf(TEXT("UENUM(BlueprintType, %s )"), *FEmitHelper::ReplaceConvertedMetaData(SourceEnum)));
 	Header.AddLine(FString::Printf(TEXT("enum class %s  : uint8"), *FEmitHelper::GetCppName(SourceEnum)));
 	Header.AddLine(TEXT("{"));
 	Header.IncreaseIndent();
@@ -388,7 +430,7 @@ FString FBlueprintCompilerCppBackendBase::GenerateCodeFromStruct(UUserDefinedStr
 	// use GetBaseFilename() so that we can coordinate #includes and filenames
 	EmitFileBeginning(FEmitHelper::GetBaseFilename(SourceStruct), EmitterContext);
 
-	EmitterContext.Header.AddLine(FString::Printf(TEXT("USTRUCT(BlueprintType %s)"), *FEmitHelper::ReplaceConvertedMetaData(SourceStruct)));
+	EmitterContext.Header.AddLine(FString::Printf(TEXT("USTRUCT(BlueprintType, %s)"), *FEmitHelper::ReplaceConvertedMetaData(SourceStruct)));
 	EmitterContext.Header.AddLine(FString::Printf(TEXT("struct %s"), *FEmitHelper::GetCppName(SourceStruct)));
 	EmitterContext.Header.AddLine(TEXT("{"));
 	EmitterContext.Header.AddLine(TEXT("public:"));
@@ -499,8 +541,19 @@ FString FBlueprintCompilerCppBackendBase::GenerateWrapperForClass(UClass* Source
 		EmitterContext.Header.AddLine(FString::Printf(TEXT("%s& GetRef__%s()"), *TypeDeclaration, *UnicodeToCPPIdentifier(Property->GetName(), false, nullptr)));
 		EmitterContext.Header.AddLine(TEXT("{"));
 		EmitterContext.Header.IncreaseIndent();
-		EmitterContext.Header.AddLine(FString::Printf(TEXT("static const FName __PropertyName(TEXT(\"%s\"));"), *Property->GetName()));
-		EmitterContext.Header.AddLine(FString::Printf(TEXT("return *(GetClass()->FindPropertyByName(__PropertyName)->ContainerPtrToValuePtr<%s>(__Object));"), *TypeDeclaration));
+
+		EmitterContext.Header.AddLine(TEXT("static TWeakObjectPtr<UProperty> __PropertyPtr{};"));
+		EmitterContext.Header.AddLine(TEXT("const UProperty* __Property = __PropertyPtr.Get();"));
+		EmitterContext.Header.AddLine(TEXT("if (nullptr == __Property)"));
+		EmitterContext.Header.AddLine(TEXT("{"));
+		EmitterContext.Header.IncreaseIndent();
+		EmitterContext.Header.AddLine(FString::Printf(TEXT("__Property = GetClass()->FindPropertyByName(FName(TEXT(\"%s\")));"), *Property->GetName()));
+		EmitterContext.Header.AddLine(TEXT("check(__Property);"));
+		EmitterContext.Header.AddLine(TEXT("__PropertyPtr = __Property;"));
+		EmitterContext.Header.DecreaseIndent();
+		EmitterContext.Header.AddLine(TEXT("}"));
+
+		EmitterContext.Header.AddLine(FString::Printf(TEXT("return *(__Property->ContainerPtrToValuePtr<%s>(__Object));"), *TypeDeclaration));
 		EmitterContext.Header.DecreaseIndent();
 		EmitterContext.Header.AddLine(TEXT("}"));
 	}
@@ -567,7 +620,7 @@ FString FBlueprintCompilerCppBackendBase::GenerateWrapperForClass(UClass* Source
 					bFirst = false;
 				}
 
-				CallFunction += UnicodeToCPPIdentifier(Property->GetName(), Property->HasAnyPropertyFlags(CPF_Deprecated), TEXT("bpv__"));
+				CallFunction += UnicodeToCPPIdentifier(Property->GetName(), Property->HasAnyPropertyFlags(CPF_Deprecated), TEXT("bpp__"));
 			}
 			CallFunction += TEXT(");");
 			EmitterContext.Header.AddLine(CallFunction);
@@ -592,11 +645,24 @@ void FBlueprintCompilerCppBackendBase::EmitFileBeginning(const FString& CleanNam
 	{
 		Dst.AddLine(FString::Printf(TEXT("#include \"%s%s\""), Message, bAddDotH ? TEXT(".h") : TEXT("")));
 	};
-	EmitIncludeHeader(EmitterContext.Body, *FEmitHelper::GetPCHFilename(), false);
+	const FString PCHFilename = FEmitHelper::GetPCHFilename();
+	if (!PCHFilename.IsEmpty())
+	{
+		EmitIncludeHeader(EmitterContext.Body, *PCHFilename, false);
+	}
+	/*
+	const FString MainHeaderFilename = FEmitHelper::GetGameMainHeaderFilename();
+	if (!MainHeaderFilename.IsEmpty() && MainHeaderFilename != PCHFilename)
+	{
+		EmitIncludeHeader(EmitterContext.Body, *MainHeaderFilename, false);
+	}
+	*/
 	EmitIncludeHeader(EmitterContext.Body, *CleanName, true);
 	EmitIncludeHeader(bIncludeCodeHelpersInHeader ? EmitterContext.Header : EmitterContext.Body, TEXT("GeneratedCodeHelpers"), true);
+	EmitIncludeHeader(EmitterContext.Header, TEXT("Blueprint/BlueprintSupport"), true);
 
 	FBackendHelperUMG::AdditionalHeaderIncludeForWidget(EmitterContext);
+	EmitterContext.Header.AddLine(FBackendHelperAnim::AddHeaders(EmitterContext.GetCurrentlyGeneratedClass()));
 
 	TSet<FString> AlreadyIncluded;
 	AlreadyIncluded.Add(CleanName);
diff --git a/Engine/Source/Developer/BlueprintCompilerCppBackend/Private/BlueprintCompilerCppBackendBase.h b/Engine/Source/Developer/BlueprintCompilerCppBackend/Private/BlueprintCompilerCppBackendBase.h
--- a/Engine/Source/Developer/BlueprintCompilerCppBackend/Private/BlueprintCompilerCppBackendBase.h
+++ b/Engine/Source/Developer/BlueprintCompilerCppBackend/Private/BlueprintCompilerCppBackendBase.h
@@ -72,6 +72,14 @@ private:
 	/** Builds both the header declaration and body implementation of a function */
 	void ConstructFunction(FKismetFunctionContext& FunctionContext, FEmitterLocalContext& EmitterContext, bool bGenerateStubOnly);
 
+	static void ConstructFunctionDeclaration(FEmitterLocalContext &EmitterContext, FKismetFunctionContext &FunctionContext, TArray<UProperty*> &ArgumentList);
+
+	static FString GenerateArgList(const FEmitterLocalContext &EmitterContext, const TArray<UProperty*> &ArgumentList);
+
+	static FString GenerateReturnType(const FEmitterLocalContext &EmitterContext, const UFunction* Function);
+
+	void ConstructFunctionBody(FEmitterLocalContext& EmitterContext, FKismetFunctionContext &FunctionContext);
+
 	void CleanBackend();
 
 	static void EmitFileBeginning(const FString& CleanName, FEmitterLocalContext& EmitterContext, bool bIncludeGeneratedH = true, bool bIncludeCodeHelpersInHeader = false);
diff --git a/Engine/Source/Developer/BlueprintCompilerCppBackend/Private/BlueprintCompilerCppBackendGatherDependencies.cpp b/Engine/Source/Developer/BlueprintCompilerCppBackend/Private/BlueprintCompilerCppBackendGatherDependencies.cpp
--- a/Engine/Source/Developer/BlueprintCompilerCppBackend/Private/BlueprintCompilerCppBackendGatherDependencies.cpp
+++ b/Engine/Source/Developer/BlueprintCompilerCppBackend/Private/BlueprintCompilerCppBackendGatherDependencies.cpp
@@ -2,18 +2,218 @@
 
 #include "BlueprintCompilerCppBackendModulePrivatePCH.h"
 #include "BlueprintCompilerCppBackendGatherDependencies.h"
-#include "BlueprintSupport.h"
+#include "IBlueprintCompilerCppBackendModule.h"
+
+struct FGatherConvertedClassDependenciesHelperBase : public FReferenceCollector
+{
+	TSet<UObject*> SerializedObjects;
+	FGatherConvertedClassDependencies& Dependencies;
+
+	FGatherConvertedClassDependenciesHelperBase(FGatherConvertedClassDependencies& InDependencies)
+		: Dependencies(InDependencies)
+	{ }
+
+	virtual bool IsIgnoringArchetypeRef() const override { return false; }
+	virtual bool IsIgnoringTransient() const override { return true; }
+
+	void FindReferences(UObject* Object)
+	{
+		{
+			FSimpleObjectReferenceCollectorArchive CollectorArchive(Object, *this);
+			CollectorArchive.SetSerializedProperty(nullptr);
+			CollectorArchive.SetFilterEditorOnly(true);
+			Object->SerializeScriptProperties(CollectorArchive);
+		}
+		Object->CallAddReferencedObjects(*this);
+	}
+
+	void FindReferencesForNewObject(UObject* Object)
+	{
+		bool AlreadyAdded = false;
+		SerializedObjects.Add(Object, &AlreadyAdded);
+		if (!AlreadyAdded)
+		{
+			FindReferences(Object);
+		}
+	}
+
+	void IncludeTheHeaderInBody(UField* InField)
+	{
+		if (InField && !Dependencies.IncludeInHeader.Contains(InField))
+		{
+			Dependencies.IncludeInBody.Add(InField);
+		}
+	}
+};
+
+struct FFindAssetsToInclude : public FGatherConvertedClassDependenciesHelperBase
+{
+	FFindAssetsToInclude(FGatherConvertedClassDependencies& InDependencies)
+		: FGatherConvertedClassDependenciesHelperBase(InDependencies)
+	{
+		FindReferences(Dependencies.GetActualStruct());
+	}
+
+	virtual void HandleObjectReference(UObject*& InObject, const UObject* InReferencingObject, const UProperty* InReferencingProperty) override
+	{
+		if (!InObject || InObject->IsA<UBlueprint>() || (InObject == Dependencies.GetActualStruct()))
+		{
+			return;
+		}
+
+		const bool bUseZConstructorInGeneratedCode = true;
+
+		//TODO: What About Delegates?
+		auto ObjAsBPGC = Cast<UBlueprintGeneratedClass>(InObject);
+		const bool bWillBeConvetedAsBPGC = ObjAsBPGC && Dependencies.WillClassBeConverted(ObjAsBPGC);
+		if (bWillBeConvetedAsBPGC)
+		{
+			if (ObjAsBPGC != Dependencies.GetActualStruct())
+			{
+				Dependencies.ConvertedClasses.Add(ObjAsBPGC);
+				if(!bUseZConstructorInGeneratedCode)
+				{
+					IncludeTheHeaderInBody(ObjAsBPGC);
+				}
+			}
+		}
+		else if (UUserDefinedStruct* UDS = Cast<UUserDefinedStruct>(InObject))
+		{
+			if (!UDS->HasAnyFlags(RF_ClassDefaultObject))
+			{
+				Dependencies.ConvertedStructs.Add(UDS);
+				if(!bUseZConstructorInGeneratedCode)
+				{
+					IncludeTheHeaderInBody(UDS);
+				}
+			}
+		}
+		else if (UUserDefinedEnum* UDE = Cast<UUserDefinedEnum>(InObject))
+		{
+			if (!UDE->HasAnyFlags(RF_ClassDefaultObject))
+			{
+				Dependencies.ConvertedEnum.Add(UDE);
+			}
+		}
+		else if ((InObject->IsAsset() || ObjAsBPGC) && !InObject->IsIn(Dependencies.GetActualStruct()))
+		{
+			// include all not converted super classes
+			for (auto SuperBPGC = ObjAsBPGC ? Cast<UBlueprintGeneratedClass>(ObjAsBPGC->GetSuperClass()) : nullptr;
+				SuperBPGC && !Dependencies.WillClassBeConverted(SuperBPGC);
+				SuperBPGC = Cast<UBlueprintGeneratedClass>(SuperBPGC->GetSuperClass()))
+			{
+				Dependencies.Assets.AddUnique(SuperBPGC);
+			}
+
+			Dependencies.Assets.AddUnique(InObject);
+			return;
+		}
+		else if (auto ObjAsClass = Cast<UClass>(InObject))
+		{
+			if (ObjAsClass->HasAnyClassFlags(CLASS_Native))
+			{
+				return;
+			}
+		}
+		else if (InObject->IsA<UScriptStruct>())
+		{
+			return;
+		}
+
+		FindReferencesForNewObject(InObject);
+	}
+};
+
+struct FFindHeadersToInclude : public FGatherConvertedClassDependenciesHelperBase
+{
+	FFindHeadersToInclude(FGatherConvertedClassDependencies& InDependencies)
+		: FGatherConvertedClassDependenciesHelperBase(InDependencies)
+	{
+		FindReferences(Dependencies.GetActualStruct());
+	}
+
+	virtual void HandleObjectReference(UObject*& InObject, const UObject* InReferencingObject, const UProperty* InReferencingProperty) override
+	{
+		const bool bIncludeMembersFromParentClass = true;
+
+		if (!InObject || InObject->IsA<UBlueprint>() || (InObject == Dependencies.GetActualStruct()))
+		{
+			return;
+		}
+
+		auto ObjAsField = Cast<UField>(InObject);
+		if (ObjAsField)
+		{
+			IncludeTheHeaderInBody(ObjAsField);
+		}
+
+		if ((InObject->IsAsset() || InObject->IsA<UBlueprintGeneratedClass>()) && !InObject->IsIn(Dependencies.GetActualStruct()))
+		{
+			return;
+		}
+		else if (auto ObjAsClass = Cast<UClass>(InObject))
+		{
+			if (bIncludeMembersFromParentClass && (ObjAsClass != UObject::StaticClass()) && Dependencies.GetActualStruct()->IsChildOf(ObjAsClass))
+			{
+				FindReferencesForNewObject(InObject);
+			}
+			if (ObjAsClass->HasAnyClassFlags(CLASS_Native))
+			{
+				return;
+			}
+		}
+		else if (InObject->IsA<UScriptStruct>())
+		{
+			return;
+		}
+
+		if ((InObject->GetOutermost() == Dependencies.GetActualStruct()->GetOutermost()) || 
+			(bIncludeMembersFromParentClass && ObjAsField && Dependencies.GetActualStruct()->IsChildOf(ObjAsField->GetOwnerClass())))
+		{
+			FindReferencesForNewObject(InObject);
+		}
+	}
+};
 
 FGatherConvertedClassDependencies::FGatherConvertedClassDependencies(UStruct* InStruct) : OriginalStruct(InStruct)
 {
 	check(OriginalStruct);
+	// Gather headers and type declarations for header.
 	DependenciesForHeader();
-	FindReferences(OriginalStruct);
+	// Gather headers (only from the class hierarchy) to include in body.
+	FFindHeadersToInclude FindHeadersToInclude(*this);
+	// Gather assets that must be referenced.
+	FFindAssetsToInclude FindAssetsToInclude(*this);
+}
+
+UClass* FGatherConvertedClassDependencies::FindOriginalClass(const UClass* InClass) const
+{
+	if (InClass)
+	{
+		IBlueprintCompilerCppBackendModule& BackEndModule = (IBlueprintCompilerCppBackendModule&)IBlueprintCompilerCppBackendModule::Get();
+		auto ClassWeakPtrPtr = BackEndModule.GetOriginalClassMap().Find(InClass);
+		UClass* OriginalClass = ClassWeakPtrPtr ? ClassWeakPtrPtr->Get() : nullptr;
+		return OriginalClass ? OriginalClass : const_cast<UClass*>(InClass);
+	}
+	return nullptr;
 }
 
 bool FGatherConvertedClassDependencies::WillClassBeConverted(const UBlueprintGeneratedClass* InClass) const
 {
-	return FReplaceCookedBPGC::Get().CouldBeConverted(InClass);
+	if (InClass && !InClass->HasAnyFlags(RF_ClassDefaultObject))
+	{
+		const UClass* ClassToCheck = FindOriginalClass(InClass);
+
+		IBlueprintCompilerCppBackendModule& BackEndModule = (IBlueprintCompilerCppBackendModule&)IBlueprintCompilerCppBackendModule::Get();
+		auto& WillBeConvertedQuery = BackEndModule.OnIsTargetedForConversionQuery();
+
+		if (WillBeConvertedQuery.IsBound())
+		{
+			return WillBeConvertedQuery.Execute(ClassToCheck);
+		}
+		return true;
+	}
+	return false;
 }
 
 void FGatherConvertedClassDependencies::DependenciesForHeader()
@@ -24,6 +224,23 @@ void FGatherConvertedClassDependencies::DependenciesForHeader()
 	TArray<UObject*> NeededObjects;
 	FReferenceFinder HeaderReferenceFinder(NeededObjects, nullptr, false, false, true, false);
 
+	auto ShouldIncludeHeaderFor = [&](UObject* InObj)->bool
+	{
+		if (InObj
+			&& (InObj->IsA<UClass>() || InObj->IsA<UEnum>() || InObj->IsA<UScriptStruct>())
+			&& !InObj->HasAnyFlags(RF_ClassDefaultObject))
+		{
+			auto ObjAsBPGC = Cast<UBlueprintGeneratedClass>(InObj);
+			const bool bWillBeConvetedAsBPGC = ObjAsBPGC && WillClassBeConverted(ObjAsBPGC);
+			const bool bRemainAsUnconvertedBPGC = ObjAsBPGC && !bWillBeConvetedAsBPGC;
+			if (!bRemainAsUnconvertedBPGC && (InObj->GetOutermost() != OriginalStruct->GetOutermost()))
+			{
+				return true;
+			}
+		}
+		return false;
+	};
+
 	for (auto Obj : ObjectsToCheck)
 	{
 		auto Property = Cast<const UProperty>(Obj);
@@ -72,97 +289,3 @@ void FGatherConvertedClassDependencies::DependenciesForHeader()
 	}
 }
 
-bool FGatherConvertedClassDependencies::ShouldIncludeHeaderFor(UObject* Obj) const
-{
-	if (Obj
-		&& (Obj->IsA<UClass>() || Obj->IsA<UEnum>() || Obj->IsA<UScriptStruct>())
-		&& !Obj->HasAnyFlags(RF_ClassDefaultObject))
-	{
-		auto ObjAsBPGC = Cast<UBlueprintGeneratedClass>(Obj);
-		const bool bWillBeConvetedAsBPGC = ObjAsBPGC && WillClassBeConverted(ObjAsBPGC);
-		const bool bRemainAsUnconvertedBPGC = ObjAsBPGC && !bWillBeConvetedAsBPGC;
-		if (!bRemainAsUnconvertedBPGC && (Obj->GetOutermost() != OriginalStruct->GetOutermost()))
-		{
-			return true;
-		}
-	}
-	return false;
-}
-
-void FGatherConvertedClassDependencies::HandleObjectReference(UObject*& InObject, const UObject* InReferencingObject, const UProperty* InReferencingProperty)
-{
-	if (!InObject || InObject->IsA<UBlueprint>() || (InObject == OriginalStruct))
-	{
-		return;
-	}
-
-	auto ObjAsField = Cast<UField>(InObject);
-	if (ObjAsField && !IncludeInHeader.Contains(ObjAsField))
-	{
-		IncludeInBody.Add(ObjAsField);
-	}
-
-	//TODO: What About Delegates?
-	auto ObjAsBPGC = Cast<UBlueprintGeneratedClass>(InObject);
-	const bool bWillBeConvetedAsBPGC = ObjAsBPGC && WillClassBeConverted(ObjAsBPGC);
-	if (bWillBeConvetedAsBPGC)
-	{
-		ConvertedClasses.Add(ObjAsBPGC);
-	}
-	else if (InObject->IsA<UUserDefinedStruct>())
-	{
-		if (!InObject->HasAnyFlags(RF_ClassDefaultObject))
-		{
-			ConvertedStructs.Add(CastChecked<UUserDefinedStruct>(InObject));
-		}
-	}
-	else if (InObject->IsA<UUserDefinedEnum>())
-	{
-		if (!InObject->HasAnyFlags(RF_ClassDefaultObject))
-		{
-			ConvertedEnum.Add(CastChecked<UUserDefinedEnum>(InObject));
-		}
-	}
-	else if ((InObject->IsAsset() || ObjAsBPGC) && !InObject->IsIn(OriginalStruct))
-	{
-		// include all not converted super classes
-		for (auto SuperBPGC = ObjAsBPGC ? Cast<UBlueprintGeneratedClass>(ObjAsBPGC->GetSuperClass()) : nullptr; 
-			SuperBPGC && !WillClassBeConverted(SuperBPGC);
-			SuperBPGC = Cast<UBlueprintGeneratedClass>(SuperBPGC->GetSuperClass()))
-		{
-			Assets.AddUnique(SuperBPGC);
-		}
-
-		Assets.AddUnique(InObject);
-		return;
-	}
-	else if (auto ObjAsClass = Cast<UClass>(InObject))
-	{
-		if (ObjAsClass->HasAnyClassFlags(CLASS_Native))
-		{
-			return;
-		}
-	}
-	else if (InObject->IsA<UScriptStruct>())
-	{
-		return;
-	}
-
-	bool AlreadyAdded = false;
-	SerializedObjects.Add(InObject, &AlreadyAdded);
-	if (!AlreadyAdded)
-	{
-		FindReferences(InObject);
-	}
-}
-
-void FGatherConvertedClassDependencies::FindReferences(UObject* Object)
-{
-	{
-		FSimpleObjectReferenceCollectorArchive CollectorArchive(Object, *this);
-		CollectorArchive.SetSerializedProperty(nullptr);
-		CollectorArchive.SetFilterEditorOnly(true);
-		Object->SerializeScriptProperties(CollectorArchive);
-	}
-	Object->CallAddReferencedObjects(*this);
-}
\ No newline at end of file
diff --git a/Engine/Source/Developer/BlueprintCompilerCppBackend/Private/BlueprintCompilerCppBackendModule.cpp b/Engine/Source/Developer/BlueprintCompilerCppBackend/Private/BlueprintCompilerCppBackendModule.cpp
--- a/Engine/Source/Developer/BlueprintCompilerCppBackend/Private/BlueprintCompilerCppBackendModule.cpp
+++ b/Engine/Source/Developer/BlueprintCompilerCppBackend/Private/BlueprintCompilerCppBackendModule.cpp
@@ -14,12 +14,16 @@ public:
 	//~ End IBlueprintCompilerCppBackendModuleInterface interface
 
 	//~ Begin IBlueprintCompilerCppBackendModule interface
-	FString ConstructBaseFilename(const UObject* AssetObj) override;
+	virtual FString ConstructBaseFilename(const UObject* AssetObj) override;
 	virtual FPCHFilenameQuery& OnPCHFilenameQuery() override;
+	virtual FIsTargetedForConversionQuery& OnIsTargetedForConversionQuery() override;
+	virtual TMap<TWeakObjectPtr<UClass>, TWeakObjectPtr<UClass> >& GetOriginalClassMap() override;
 	//~ End IBlueprintCompilerCppBackendModule interface
 
 private: 
 	FPCHFilenameQuery PCHFilenameQuery;
+	FIsTargetedForConversionQuery IsTargetedForConversionQuery;
+	TMap<TWeakObjectPtr<UClass>, TWeakObjectPtr<UClass> > OriginalClassMap;
 };
 
 IBlueprintCompilerCppBackend* FBlueprintCompilerCppBackendModule::Create()
@@ -38,4 +42,14 @@ IBlueprintCompilerCppBackendModule::FPCHFilenameQuery& FBlueprintCompilerCppBack
 	return PCHFilenameQuery;
 }
 
+IBlueprintCompilerCppBackendModule::FIsTargetedForConversionQuery& FBlueprintCompilerCppBackendModule::OnIsTargetedForConversionQuery()
+{
+	return IsTargetedForConversionQuery;
+}
+
+TMap<TWeakObjectPtr<UClass>, TWeakObjectPtr<UClass> >& FBlueprintCompilerCppBackendModule::GetOriginalClassMap()
+{
+	return OriginalClassMap;
+}
+
 IMPLEMENT_MODULE(FBlueprintCompilerCppBackendModule, BlueprintCompilerCppBackend)
