Copying //UE4/Dev-Blueprints to Dev-Main (//UE4/Dev-Main)

==========================
MAJOR FEATURES + CHANGES
==========================

Change 2806214 on 2015/12/16 by [REDACTED_PERSON]

	Connection draw policy refactored into a factory system.

	PR#1663

Change 2806845 on 2015/12/17 by [REDACTED_PERSON]

	Blueprint C++ Conversion: fixed varioaus problem.
	Some limitations of included header list were reverted :(
	Any compilation-time optimization are premature (as far as Orion cannot be compiled).

Change 2806892 on 2015/12/17 by [REDACTED_PERSON]

	Blueprint C++ Conversion: fixed compilation error "fatal error C1001: An internal error has occurred in the compiler."

Change 2807020 on 2015/12/17 by [REDACTED_PERSON]

	Recursively expand all tree items in graph context menus or My Blueprint window by holding shift when clicking on an arrow.

Change 2807639 on 2015/12/17 by [REDACTED_PERSON]

	BP C++ conversion: Fix for const-correctness. Included generated headers in structs.

Change 2807880 on 2015/12/17 by [REDACTED_PERSON]

	PR #1865 : Render Kismet graph pin color box with alpha
	https://github.com/EpicGames/UnrealEngine/pull/1865

	#github 1865
	#1865

Change 2808538 on 2015/12/18 by [REDACTED_PERSON]

	Workaround for UE-24467. Some corrupted files can be opened and fixed.

	#codereivew Dan.Oconnor, Michael.Schoell

Change 2808540 on 2015/12/18 by [REDACTED_PERSON]

	BP C++ conversion: fixed const-correctness related issues

Change 2809333 on 2015/12/18 by [REDACTED_PERSON]

	[UE-23452] SCS/UCS component instancing optimizations

	change summary:
	- added FArchive::FCustomPropertyListNode
	- modified various parts of serialization code to accomodate custom property lists
	- added cooked component instancing data + supporting APIs to UBlueprintGeneratedClass

Change 2810216 on 2015/12/21 by [REDACTED_PERSON]

	Blueprint C++ Conversion:
	- Fixed Compiler Error C2026 (too big string)
	- Fixed a lot of errors related to const correctness. "NativeConst" and "NativeConstTemplateArg" MetaData added.
	- Fixed bunch of problems with TEnumAsByte
	- Fixed for error C2059: syntax error : 'bad suffix on number' - when struct name starts with a digit
	- Fixed int -> bool  conversion

Change 2810397 on 2015/12/21 by [REDACTED_PERSON]

	Blueprint C++ Conversion: Fixed Compiler Error C1091 (too big string)

Change 2810638 on 2015/12/21 by [REDACTED_PERSON]

	Timers for measuring VM overhead, stats for tracking individual script functions improved (now nicludes ProcessEvent). Stats system is expected to avoid double counting for object and function stats.

Change 2811038 on 2015/12/22 by [REDACTED_PERSON]

	[UE-23452] Fix uninitialized variables in cooked component data.

	- should eliminate crashes introduced by last change

Change 2811054 on 2015/12/22 by [REDACTED_PERSON]

	[UE-23452] Additional fix for uninitialized USTRUCT property.

Change 2811254 on 2015/12/22 by [REDACTED_PERSON]

	More script function detail in stats

Change 2811325 on 2015/12/22 by [REDACTED_PERSON]

	Blueprint C++ Conversion: improved: ExcludedBlueprintTypes list

Change 2812316 on 2015/12/24 by [REDACTED_PERSON]

	Added more ByValue/ByRef test cases and added a way to dump all tests that are no longer succeeding.

	Tests Include:
	Verification that functions take and can modify by-ref values.
	ForEach loop verification. Tests for verifying that current functionality will continue to work after COPY injection and that "expected" functionality will work (and currently doesn't).
	More MegaArray tests.

Change 2812318 on 2015/12/24 by [REDACTED_PERSON]

	Usability fix: Duplicating graphs in the MyBP window will now open and focus the new graph.

Change 2813179 on 2015/12/29 by [REDACTED_PERSON]

	When promoting pins to variables, will no longer carry bIsReference, bIsConst, or bIsWeakPointer value to the new variable's type.

Change 2813435 on 2015/12/30 by [REDACTED_PERSON]

	Submitting by-value/by-ref testing BP with more tests:

	More "Set Members..." tests to catch edge cases.
	More Array tests.

Change 2813441 on 2015/12/30 by [REDACTED_PERSON]

diff --git a/Engine/Config/BaseEditor.ini b/Engine/Config/BaseEditor.ini
--- a/Engine/Config/BaseEditor.ini
+++ b/Engine/Config/BaseEditor.ini
@@ -463,3 +463,19 @@ AngleUnits=EUnit::Degrees
 SpeedUnits=EUnit::MetersPerSecond
 TemperatureUnits=EUnit::Celsius
 ForceUnits=EUnit::Newtons
+
+[BlueprintNativizationSettings]
+bUseStaticVariablesInWrappers=true
+bUseStaticVariablesInClasses=true
+bSkipUFUNCTION=false
+bUsePRAGMA_DISABLE_OPTIMIZATION=false
+
+; this flag can change BP behaviour, see UE-23629. It doesn't help in image size reduction.
+bUseLocalEventGraphVariables=false
+
+;+ExcludedBlueprintTypes=/Script/UMGEditor.WidgetBlueprint
+;+ExcludedBlueprintTypes=/Script/Engine.AnimBlueprint
+
++ExcludedAssets=/Engine/Tutorial/BlueprintTutorials/TutorialAssets/Tutorial_BP_Interface
++ExcludedAssets=/Engine/Tutorial/SubEditors/TutorialAssets/Character/TutorialCharacter
++ExcludedAssets=/Engine/Tutorial/SubEditors/TutorialAssets/Character/TutorialTPP_AnimBlueprint
\ No newline at end of file
diff --git a/Engine/Source/Developer/AssetTools/Private/AssetTypeActions/AssetTypeActions_Blueprint.cpp b/Engine/Source/Developer/AssetTools/Private/AssetTypeActions/AssetTypeActions_Blueprint.cpp
--- a/Engine/Source/Developer/AssetTools/Private/AssetTypeActions/AssetTypeActions_Blueprint.cpp
+++ b/Engine/Source/Developer/AssetTools/Private/AssetTypeActions/AssetTypeActions_Blueprint.cpp
@@ -71,10 +71,18 @@ void FAssetTypeActions_Blueprint::OpenAssetEditor( const TArray<UObject*>& InObj
 	for (auto ObjIt = InObjects.CreateConstIterator(); ObjIt; ++ObjIt)
 	{
 		auto Blueprint = Cast<UBlueprint>(*ObjIt);
-		if (Blueprint && Blueprint->SkeletonGeneratedClass && Blueprint->GeneratedClass )
+		if (Blueprint)
 		{
-			FBlueprintEditorModule& BlueprintEditorModule = FModuleManager::LoadModuleChecked<FBlueprintEditorModule>( "Kismet" );
-			TSharedRef< IBlueprintEditor > NewKismetEditor = BlueprintEditorModule.CreateBlueprintEditor( Mode, EditWithinLevelEditor, Blueprint, ShouldUseDataOnlyEditor(Blueprint) );
+			bool bLetOpen = true;
+			if (!Blueprint->SkeletonGeneratedClass || !Blueprint->GeneratedClass)
+			{
+				bLetOpen = EAppReturnType::Yes == FMessageDialog::Open(EAppMsgType::YesNoCancel, LOCTEXT("FailedToLoadBlueprint", "Blueprint could not be loaded because it derives from an invalid class.  Check to make sure the parent class for this blueprint hasn't been removed! Do you want to continue (it can crash the editor)?"));
+			}
+			if (bLetOpen)
+			{
+				FBlueprintEditorModule& BlueprintEditorModule = FModuleManager::LoadModuleChecked<FBlueprintEditorModule>("Kismet");
+				TSharedRef< IBlueprintEditor > NewKismetEditor = BlueprintEditorModule.CreateBlueprintEditor(Mode, EditWithinLevelEditor, Blueprint, ShouldUseDataOnlyEditor(Blueprint));
+			}
 		}
 		else
 		{
diff --git a/Engine/Source/Developer/BlueprintCompilerCppBackend/Private/BlueprintCompilerCppBackend.cpp b/Engine/Source/Developer/BlueprintCompilerCppBackend/Private/BlueprintCompilerCppBackend.cpp
--- a/Engine/Source/Developer/BlueprintCompilerCppBackend/Private/BlueprintCompilerCppBackend.cpp
+++ b/Engine/Source/Developer/BlueprintCompilerCppBackend/Private/BlueprintCompilerCppBackend.cpp
@@ -322,10 +322,10 @@ FString FBlueprintCompilerCppBackend::EmitSwitchValueStatmentInner(FEmitterLocal
 		| EPropertyExportCPPFlags::CPPF_BlueprintCppBackend;
 
 	check(IndexTerm && IndexTerm->AssociatedVarProperty);
-	const FString IndexDeclaration = EmitterContext.ExportCppDeclaration(IndexTerm->AssociatedVarProperty, EExportedDeclaration::Parameter, CppTemplateTypeFlags, true);
+	const FString IndexDeclaration = EmitterContext.ExportCppDeclaration(IndexTerm->AssociatedVarProperty, EExportedDeclaration::Local, CppTemplateTypeFlags, true);
 
 	check(DefaultValueTerm && DefaultValueTerm->AssociatedVarProperty);
-	const FString ValueDeclaration = EmitterContext.ExportCppDeclaration(DefaultValueTerm->AssociatedVarProperty, EExportedDeclaration::Parameter, CppTemplateTypeFlags, true);
+	const FString ValueDeclaration = EmitterContext.ExportCppDeclaration(DefaultValueTerm->AssociatedVarProperty, EExportedDeclaration::Local, CppTemplateTypeFlags, true);
 
 	FString Result = FString::Printf(TEXT("TSwitchValue<%s, %s>(%s, %s, %d")
 		, *IndexDeclaration
@@ -562,14 +562,15 @@ FString FBlueprintCompilerCppBackend::TermToText(FEmitterLocalContext& EmitterCo
 		}
 
 		FString ResultPath;
-
+		const bool bNativeConst = Term->AssociatedVarProperty && Term->AssociatedVarProperty->HasMetaData(FName(TEXT("NativeConst")));
+		bool bIsAccessible = bGetter || !bNativeConst;
 		if (Term->Context && Term->Context->IsStructContextType())
 		{
 			check(Term->AssociatedVarProperty);
-			const bool bIsAccessible = !Term->AssociatedVarProperty->HasAnyPropertyFlags(CPF_NativeAccessSpecifierPrivate | CPF_NativeAccessSpecifierProtected);
+			bIsAccessible &= !Term->AssociatedVarProperty->HasAnyPropertyFlags(CPF_NativeAccessSpecifierPrivate | CPF_NativeAccessSpecifierProtected);
 			if (!bIsAccessible)
 			{
-				ResultPath = FEmitHelper::AccessInaccessibleProperty(EmitterContext, Term->AssociatedVarProperty, ContextStr, TEXT("&"));
+				ResultPath = FEmitHelper::AccessInaccessibleProperty(EmitterContext, Term->AssociatedVarProperty, ContextStr, TEXT("&"), 0, bGetter);
 			}
 			else
 			{
@@ -580,7 +581,7 @@ FString FBlueprintCompilerCppBackend::TermToText(FEmitterLocalContext& EmitterCo
 		{
 			const bool bSelfContext = (!Term->Context) || (Term->Context->Name == PSC_Self);
 			const bool bPropertyOfParent = EmitterContext.Dependencies.GetActualStruct()->IsChildOf(Term->AssociatedVarProperty->GetOwnerStruct());
-			const bool bIsAccessible = !Term->AssociatedVarProperty->HasAnyPropertyFlags(CPF_NativeAccessSpecifierPrivate)
+			bIsAccessible &= !Term->AssociatedVarProperty->HasAnyPropertyFlags(CPF_NativeAccessSpecifierPrivate)
 				&& ((bPropertyOfParent && bSelfContext) || !Term->AssociatedVarProperty->HasAnyPropertyFlags(CPF_NativeAccessSpecifierProtected));
 
 			auto MinimalClass = Term->AssociatedVarProperty->GetOwnerClass();
@@ -605,7 +606,7 @@ FString FBlueprintCompilerCppBackend::TermToText(FEmitterLocalContext& EmitterCo
 					ensure(ContextStr.IsEmpty());
 					ContextStr = TEXT("this");
 				}
-				ResultPath = FEmitHelper::AccessInaccessibleProperty(EmitterContext, Term->AssociatedVarProperty, ContextStr, FString());
+				ResultPath = FEmitHelper::AccessInaccessibleProperty(EmitterContext, Term->AssociatedVarProperty, ContextStr, FString(), 0, bGetter);
 			}
 			else
 			{
@@ -627,6 +628,24 @@ FString FBlueprintCompilerCppBackend::TermToText(FEmitterLocalContext& EmitterCo
 			ResultPath += TEXT(".Get()");
 		}
 
+		const bool bNativeConstTemplateArg = Term->AssociatedVarProperty && Term->AssociatedVarProperty->HasMetaData(FName(TEXT("NativeConstTemplateArg")));
+		if ((bNativeConst || bNativeConstTemplateArg) && bIsAccessible && bGetter)
+		{
+			if (Term->Type.bIsArray && bNativeConstTemplateArg)
+			{
+				FEdGraphPinType InnerType = Term->Type;
+				InnerType.bIsArray = false;
+				InnerType.bIsConst = false;
+				const FString CppType = FEmitHelper::PinTypeToNativeType(InnerType);
+				ResultPath = FString::Printf(TEXT("TArrayCaster<const %s>(%s).Get<%s>()"), *CppType, *ResultPath, *CppType);
+			}
+			else
+			{
+				const FString CppType = FEmitHelper::PinTypeToNativeType(Term->Type);
+				ResultPath = FString::Printf(TEXT("const_cast<%s>(%s)"), *CppType, *ResultPath);
+			}
+		}
+
 		FString Conditions;
 		if (bUseSafeContext)
 		{
@@ -748,7 +767,7 @@ void FBlueprintCompilerCppBackend::InnerFunctionImplementation(FKismetFunctionCo
 				EmitEndOfThreadStatement(EmitterContext, FunctionContext);
 				break;
 			case KCST_Comment:
-				EmitterContext.AddLine(FString::Printf(TEXT("// %s"), *Statement.Comment));
+				EmitterContext.AddLine(FString::Printf(TEXT("// %s"), *Statement.Comment.Replace(TEXT("\n"), TEXT(" "))));
 				break;
 			case KCST_DebugSite:
 				break;
diff --git a/Engine/Source/Developer/BlueprintCompilerCppBackend/Private/BlueprintCompilerCppBackendBase.cpp b/Engine/Source/Developer/BlueprintCompilerCppBackend/Private/BlueprintCompilerCppBackendBase.cpp
--- a/Engine/Source/Developer/BlueprintCompilerCppBackend/Private/BlueprintCompilerCppBackendBase.cpp
+++ b/Engine/Source/Developer/BlueprintCompilerCppBackend/Private/BlueprintCompilerCppBackendBase.cpp
@@ -4,6 +4,51 @@
 #include "BlueprintCompilerCppBackendBase.h"
 #include "BlueprintCompilerCppBackendUtils.h"
 
+TArray<class UFunction*> IBlueprintCompilerCppBackendModule::CollectBoundFunctions(class UBlueprint* BP)
+{
+	// it would be clean to recover info about delegates from bytecode, but then the owner class is unknown, that's why we check the nodes.
+	TArray<class UFunction*> Result;
+	check(BP);
+	TArray<UEdGraph*> Graphs;
+	BP->GetAllGraphs(Graphs);
+	for (UEdGraph* Graph : Graphs)
+	{
+		if (!Graph)
+		{
+			continue;
+		}
+		for (UEdGraphNode* Node : Graph->Nodes)
+		{
+			FName FunctionName = NAME_None;
+			UClass* FunctionOwnerClass = nullptr;
+			if (UK2Node_CreateDelegate* CreateDelegate = Cast<UK2Node_CreateDelegate>(Node))
+			{
+				FunctionName = CreateDelegate->GetFunctionName();
+				FunctionOwnerClass = CreateDelegate->GetScopeClass(true);
+			}
+			else if (UK2Node_Event* EventNode = Cast<UK2Node_Event>(Node))
+			{
+				UEdGraphPin* DelegateOutPin = EventNode->FindPin(UK2Node_Event::DelegateOutputName);
+				if (DelegateOutPin && DelegateOutPin->LinkedTo.Num())
+				{
+					FunctionOwnerClass = BP->GeneratedClass;
+					FunctionName = EventNode->GetFunctionName();
+				}
+			}
+
+			FunctionOwnerClass = FunctionOwnerClass ? FunctionOwnerClass->GetAuthoritativeClass() : nullptr;
+			UFunction* Func = FunctionOwnerClass ? FunctionOwnerClass->FindFunctionByName(FunctionName) : nullptr;
+			Func = Func ? FEmitHelper::GetOriginalFunction(Func) : nullptr;
+			if (Func)
+			{
+				Result.Add(Func);
+			}
+		}
+	}
+
+	return Result;
+}
+
 void FBlueprintCompilerCppBackendBase::EmitStructProperties(FEmitterLocalContext& EmitterContext, UStruct* SourceClass)
 {
 	// Emit class variables
@@ -25,7 +70,7 @@ void FBlueprintCompilerCppBackendBase::EmitStructProperties(FEmitterLocalContext
 		PropertyMacro += TEXT(")");
 		EmitterContext.Header.AddLine(PropertyMacro);
 
-		const FString CppDeclaration = EmitterContext.ExportCppDeclaration(Property, EExportedDeclaration::Member, EPropertyExportCPPFlags::CPPF_CustomTypeName | EPropertyExportCPPFlags::CPPF_BlueprintCppBackend);
+		const FString CppDeclaration = EmitterContext.ExportCppDeclaration(Property, EExportedDeclaration::Member, EPropertyExportCPPFlags::CPPF_CustomTypeName | EPropertyExportCPPFlags::CPPF_BlueprintCppBackend | EPropertyExportCPPFlags::CPPF_NoConst);
 		EmitterContext.Header.AddLine(CppDeclaration + TEXT(";"));
 	}
 }
@@ -110,7 +155,33 @@ FString FBlueprintCompilerCppBackendBase::GenerateCodeFromClass(UClass* SourceCl
 	
 	FGatherConvertedClassDependencies Dependencies(SourceClass);
 	FEmitterLocalContext EmitterContext(Dependencies);
+
 	EmitFileBeginning(CleanCppClassName, EmitterContext);
+	{
+		IBlueprintCompilerCppBackendModule& BackEndModule = (IBlueprintCompilerCppBackendModule&)IBlueprintCompilerCppBackendModule::Get();
+
+		auto MarkAsNecessary = [](IBlueprintCompilerCppBackendModule& InBackEndModule, const FGatherConvertedClassDependencies InDependencies, UField* InField)
+		{
+			UBlueprintGeneratedClass* BPGC = Cast<UBlueprintGeneratedClass>(InField);
+			UBlueprint* BP = (BPGC && !InDependencies.WillClassBeConverted(BPGC)) ? Cast<UBlueprint>(BPGC->ClassGeneratedBy) : nullptr;
+			if (BP)
+			{
+				InBackEndModule.OnIncludingUnconvertedBP().ExecuteIfBound(BP);
+			}
+		};
+		for (UField* Field : Dependencies.IncludeInHeader)
+		{
+			MarkAsNecessary(BackEndModule, Dependencies, Field);
+		}
+		for (UField* Field : Dependencies.DeclareInHeader)
+		{
+			MarkAsNecessary(BackEndModule, Dependencies, Field);
+		}
+		for (UField* Field : Dependencies.IncludeInBody)
+		{
+			MarkAsNecessary(BackEndModule, Dependencies, Field);
+		}
+	}
 
 	UClass* OriginalSourceClass = Dependencies.FindOriginalClass(SourceClass);
 	ensure(OriginalSourceClass != SourceClass);
@@ -227,18 +298,62 @@ void FBlueprintCompilerCppBackendBase::ConstructFunction(FKismetFunctionContext&
 		PropertyDest.Add(Property);
 	}
 
-	if (FunctionContext.bIsUbergraph)
+	static const FBoolConfigValueHelper UsePRAGMA_DISABLE_OPTIMIZATION(TEXT("BlueprintNativizationSettings"), TEXT("bUsePRAGMA_DISABLE_OPTIMIZATION"));
+	if (FunctionContext.bIsUbergraph && UsePRAGMA_DISABLE_OPTIMIZATION)
 	{
 		EmitterContext.AddLine(TEXT("PRAGMA_DISABLE_OPTIMIZATION"));
 	}
 
 	ConstructFunctionDeclaration(EmitterContext, FunctionContext, ArgumentList);
 
+	const bool bIsConstFunction = FunctionContext.Function->HasAllFunctionFlags(FUNC_Const);
+	if (bIsConstFunction)
+	{
+		const FString InnerImplementationFunctionName = FString::Printf(TEXT("%s_Inner_%d")
+			, *FEmitHelper::GetCppName(FunctionContext.Function)
+			, FEmitHelper::GetInheritenceLevel(FunctionContext.Function->GetOwnerStruct()));
+
+		const FString ReturnType = GenerateReturnType(EmitterContext, FunctionContext.Function);
+		const FString ArgList = GenerateArgList(EmitterContext, ArgumentList);
+		const FString ArgListNoTypes = GenerateArgList(EmitterContext, ArgumentList, true);
+		const FString ClassCppName = FEmitHelper::GetCppName(EmitterContext.GetCurrentlyGeneratedClass());
+
+		// Inner header declaration
+		EmitterContext.Header.AddLine(FString::Printf(TEXT("%s %s%s;")
+			, *ReturnType, *InnerImplementationFunctionName, *ArgList));
+
+		//Original implementation
+		EmitterContext.AddLine(TEXT("{"));
+		EmitterContext.IncreaseIndent();
+		EmitterContext.AddLine(FString::Printf(TEXT("%sconst_cast<%s*>(this)->%s%s;")
+			, FunctionContext.Function->GetReturnProperty() ? TEXT("return ") : TEXT("")
+			, *ClassCppName
+			, *InnerImplementationFunctionName
+			, *ArgListNoTypes));
+		EmitterContext.DecreaseIndent();
+		EmitterContext.AddLine(TEXT("}"));
+
+		// Inner body declaration
+		EmitterContext.AddLine(FString::Printf(TEXT("%s %s::%s%s")
+			, *ReturnType, *ClassCppName, *InnerImplementationFunctionName, *ArgList));
+
+	}
+
 	// Start the body of the implementation
 	EmitterContext.AddLine(TEXT("{"));
 	EmitterContext.IncreaseIndent();
 	if (!bGenerateStubOnly)
 	{
+		for (UProperty* Property : ArgumentList)
+		{
+			if (FEmitHelper::PropertyForConstCast(Property))
+			{
+				const uint32 ExportFlags = EPropertyExportCPPFlags::CPPF_CustomTypeName | EPropertyExportCPPFlags::CPPF_BlueprintCppBackend | EPropertyExportCPPFlags::CPPF_NoConst;
+				const FString ActualArg = EmitterContext.ExportCppDeclaration(Property, EExportedDeclaration::Parameter, ExportFlags, false);
+				const FString NoConstType = EmitterContext.ExportCppDeclaration(Property, EExportedDeclaration::Parameter, ExportFlags, true);
+				EmitterContext.AddLine(FString::Printf(TEXT("%s = const_cast<%s>(%s__const);"), *ActualArg, *NoConstType, *FEmitHelper::GetCppName(Property)));
+			}
+		}
 		DeclareLocalVariables(EmitterContext, LocalVariables);
 		ConstructFunctionBody(EmitterContext, FunctionContext);
 	}
@@ -249,7 +364,7 @@ void FBlueprintCompilerCppBackendBase::ConstructFunction(FKismetFunctionContext&
 	EmitterContext.DecreaseIndent();
 	EmitterContext.AddLine(TEXT("}"));
 
-	if (FunctionContext.bIsUbergraph)
+	if (FunctionContext.bIsUbergraph && UsePRAGMA_DISABLE_OPTIMIZATION)
 	{
 		EmitterContext.AddLine(TEXT("PRAGMA_ENABLE_OPTIMIZATION"));
 	}
@@ -267,7 +382,7 @@ void FBlueprintCompilerCppBackendBase::ConstructFunctionDeclaration(FEmitterLoca
 
 	{
 		UFunction* const Function = FunctionContext.Function;
-		const UFunction* const OriginalFunction = FEmitHelper::GetOriginalFunction(Function);
+		UFunction* const OriginalFunction = FEmitHelper::GetOriginalFunction(Function);
 		TArray<FString> AdditionalMetaData;
 		TArray<FString> AdditionalTags;
 		bool bGenerateAsNativeEventImplementation = false;
@@ -300,11 +415,52 @@ void FBlueprintCompilerCppBackendBase::ConstructFunctionDeclaration(FEmitterLoca
 		ensure(!bIsVirtual || Function->IsSignatureCompatibleWith(OriginalFunction));
 		bIsOverride = bGenerateAsNativeEventImplementation || (bIsVirtual && (Function != OriginalFunction));
 
+		auto PreliminaryConditionsToSkipMacroUFUNC = [](UFunction* InFunction) -> bool 
+		{
+			check(InFunction);
+			return !FEmitHelper::ShouldHandleAsNativeEvent(InFunction)
+				&& !FEmitHelper::ShouldHandleAsImplementableEvent(InFunction)
+				&& !InFunction->GetOwnerClass()->IsChildOf<UInterface>()
+				&& !InFunction->HasAnyFunctionFlags(FUNC_Exec | FUNC_Static | FUNC_Native
+					| FUNC_Net | FUNC_NetServer | FUNC_NetClient | FUNC_NetMulticast | FUNC_NetReliable
+					| FUNC_BlueprintAuthorityOnly | FUNC_BlueprintCosmetic | FUNC_NetValidate
+					| FUNC_MulticastDelegate | FUNC_Delegate);
+		};
+
+		auto FunctionIsBoundToAnyDelegate = [](UFunction* InFunction)
+		{
+			check(InFunction);
+			IBlueprintCompilerCppBackendModule& BackEndModule = (IBlueprintCompilerCppBackendModule&)IBlueprintCompilerCppBackendModule::Get();
+			auto& IsFunctionUsedInADelegate = BackEndModule.GetIsFunctionUsedInADelegateCallback();
+			return ensure(IsFunctionUsedInADelegate.IsBound()) ? IsFunctionUsedInADelegate.Execute(InFunction) : true;
+		};
+
+		auto IsFunctionUsedByReplication = [](UFunction* InFunction)
+		{
+			check(InFunction);
+			for (UProperty* Prop : TFieldRange<UProperty>(InFunction->GetOwnerClass(), EFieldIteratorFlags::ExcludeSuper))
+			{
+				if (Prop && FEmitHelper::HasAllFlags(Prop->PropertyFlags, CPF_Net | CPF_RepNotify) && (Prop->RepNotifyFunc == InFunction->GetFName()))
+				{
+					return true;
+				}
+			}
+			return false;
+		};
+
+		static const FBoolConfigValueHelper TryToSkipMacroUFUNC(TEXT("BlueprintNativizationSettings"), TEXT("bSkipUFUNCTION"));
+		const bool bTryToSkipMacroUFUNC = TryToSkipMacroUFUNC; // should be enabled only when all BP are compiled.
+		const bool bSkipMacro = bTryToSkipMacroUFUNC
+			&& PreliminaryConditionsToSkipMacroUFUNC(Function)
+			&& ((Function == OriginalFunction) || PreliminaryConditionsToSkipMacroUFUNC(OriginalFunction))
+			&& !IsFunctionUsedByReplication(Function)
+			&& !FunctionIsBoundToAnyDelegate(OriginalFunction);
+
 		if (bNetImplementation && bIsOverride)
 		{
 			FunctionHeaderName = FunctionBodyName;
 		}
-		else if (!bGenerateAsNativeEventImplementation)
+		else if (!bGenerateAsNativeEventImplementation && !bSkipMacro)
 		{
 			EmitterContext.Header.AddLine(FEmitHelper::EmitUFuntion(Function, AdditionalTags, AdditionalMetaData));
 		}
@@ -330,7 +486,7 @@ void FBlueprintCompilerCppBackendBase::ConstructFunctionDeclaration(FEmitterLoca
 		, bIsOverride ? TEXT(" override") : TEXT("")));
 }
 
-FString FBlueprintCompilerCppBackendBase::GenerateArgList(const FEmitterLocalContext &EmitterContext, const TArray<UProperty*> &ArgumentList)
+FString FBlueprintCompilerCppBackendBase::GenerateArgList(const FEmitterLocalContext &EmitterContext, const TArray<UProperty*> &ArgumentList, bool bOnlyParamName)
 {
 	FString ArgListStr = TEXT("(");
 	for (int32 i = 0; i < ArgumentList.Num(); ++i)
@@ -342,12 +498,23 @@ FString FBlueprintCompilerCppBackendBase::GenerateArgList(const FEmitterLocalCon
 			ArgListStr += TEXT(", ");
 		}
 
-		if (ArgProperty->HasAnyPropertyFlags(CPF_OutParm))
+		const FString NamePostFix = FEmitHelper::PropertyForConstCast(ArgProperty) ? TEXT("__const") : TEXT("");
+		if (bOnlyParamName)
 		{
-			ArgListStr += TEXT("/*out*/ ");
+			ArgListStr += FEmitHelper::GetCppName(ArgProperty) + NamePostFix;
+		}
+		else
+		{
+			if (ArgProperty->HasAnyPropertyFlags(CPF_OutParm))
+			{
+				ArgListStr += TEXT("/*out*/ ");
+			}
+			ArgListStr += EmitterContext.ExportCppDeclaration(ArgProperty
+				, EExportedDeclaration::Parameter
+				, EPropertyExportCPPFlags::CPPF_CustomTypeName | EPropertyExportCPPFlags::CPPF_BlueprintCppBackend | EPropertyExportCPPFlags::CPPF_ArgumentOrReturnValue
+				, false
+				, NamePostFix);
 		}
-
-		ArgListStr += EmitterContext.ExportCppDeclaration(ArgProperty, EExportedDeclaration::Parameter, EPropertyExportCPPFlags::CPPF_CustomTypeName | EPropertyExportCPPFlags::CPPF_BlueprintCppBackend);
 	}
 	ArgListStr += TEXT(")");
 	return ArgListStr;
@@ -362,7 +529,8 @@ FString FBlueprintCompilerCppBackendBase::GenerateReturnType(const FEmitterLocal
 			| EPropertyExportCPPFlags::CPPF_NoConst
 			| EPropertyExportCPPFlags::CPPF_NoRef
 			| EPropertyExportCPPFlags::CPPF_NoStaticArray
-			| EPropertyExportCPPFlags::CPPF_BlueprintCppBackend;
+			| EPropertyExportCPPFlags::CPPF_BlueprintCppBackend
+			| EPropertyExportCPPFlags::CPPF_ArgumentOrReturnValue;
 		return EmitterContext.ExportCppDeclaration(ReturnValue, EExportedDeclaration::Parameter, LocalExportCPPFlags, true);
 	}
 	return TEXT("void");
@@ -428,7 +596,7 @@ FString FBlueprintCompilerCppBackendBase::GenerateCodeFromStruct(UUserDefinedStr
 	FGatherConvertedClassDependencies Dependencies(SourceStruct);
 	FEmitterLocalContext EmitterContext(Dependencies);
 	// use GetBaseFilename() so that we can coordinate #includes and filenames
-	EmitFileBeginning(FEmitHelper::GetBaseFilename(SourceStruct), EmitterContext);
+	EmitFileBeginning(FEmitHelper::GetBaseFilename(SourceStruct), EmitterContext, true, true);
 
 	EmitterContext.Header.AddLine(FString::Printf(TEXT("USTRUCT(BlueprintType, %s)"), *FEmitHelper::ReplaceConvertedMetaData(SourceStruct)));
 	EmitterContext.Header.AddLine(FString::Printf(TEXT("struct %s"), *FEmitHelper::GetCppName(SourceStruct)));
@@ -470,22 +638,18 @@ FString FBlueprintCompilerCppBackendBase::GenerateWrapperForClass(UClass* Source
 		MCDelegateProperties.Add(MCDelegateProp);
 	}
 
-	const bool bGenerateAnyFunction = (0 != FunctionsToGenerate.Num()) || (0 != MCDelegateProperties.Num());
+	const bool bGenerateAnyMCDelegateProperty = (0 != MCDelegateProperties.Num());
 
 	// Include standard stuff
-	EmitFileBeginning(FEmitHelper::GetBaseFilename(SourceClass), EmitterContext, bGenerateAnyFunction, true);
+	EmitFileBeginning(FEmitHelper::GetBaseFilename(SourceClass), EmitterContext, bGenerateAnyMCDelegateProperty, true, true);
 
 	// DELEGATES
 	const FString DelegatesClassName = FString::Printf(TEXT("U__Delegates__%s"), *FEmitHelper::GetCppName(SourceClass));
-	auto GenerateDelegateTypeName = [](UFunction* Func) -> FString
-	{
-		return FString::Printf(TEXT("F__Delegate__%s"), *FEmitHelper::GetCppName(Func));
-	};
 	auto GenerateMulticastDelegateTypeName = [](UMulticastDelegateProperty* MCDelegateProp) -> FString
 	{
 		return FString::Printf(TEXT("F__MulticastDelegate__%s"), *FEmitHelper::GetCppName(MCDelegateProp));
 	};
-	if (bGenerateAnyFunction)
+	if (bGenerateAnyMCDelegateProperty)
 	{
 		EmitterContext.Header.AddLine(TEXT("UCLASS()"));
 		EmitterContext.Header.AddLine(FString::Printf(TEXT("class %s : public UObject"), *DelegatesClassName));
@@ -493,12 +657,6 @@ FString FBlueprintCompilerCppBackendBase::GenerateWrapperForClass(UClass* Source
 		EmitterContext.Header.AddLine(TEXT("public:"));
 		EmitterContext.Header.IncreaseIndent();
 		EmitterContext.Header.AddLine(TEXT("GENERATED_BODY()"));
-		for (auto Func : FunctionsToGenerate)
-		{
-			FString ParamNumberStr, Parameters;
-			FEmitHelper::ParseDelegateDetails(EmitterContext, Func, Parameters, ParamNumberStr);
-			EmitterContext.Header.AddLine(FString::Printf(TEXT("DECLARE_DYNAMIC_DELEGATE%s(%s%s);"), *ParamNumberStr, *GenerateDelegateTypeName(Func), *Parameters));
-		}
 		for (auto MCDelegateProp : MCDelegateProperties)
 		{
 			FString ParamNumberStr, Parameters;
@@ -531,6 +689,9 @@ FString FBlueprintCompilerCppBackendBase::GenerateWrapperForClass(UClass* Source
 	// Constructor
 	EmitterContext.Header.AddLine(FString::Printf(TEXT("%s(const %s* __InObject) : %s(__InObject){}"), *WrapperName, *FEmitHelper::GetCppName(EmitterContext.GetFirstNativeOrConvertedClass(SourceClass)), *ParentStruct));
 
+	static const FBoolConfigValueHelper UseStaticVariables(TEXT("BlueprintNativizationSettings"), TEXT("bUseStaticVariablesInWrappers"));
+	const bool bUseStaticVariables = UseStaticVariables;
+
 	// PROPERTIES:
 	for (auto Property : TFieldRange<UProperty>(SourceClass, EFieldIteratorFlags::ExcludeSuper))
 	{
@@ -541,18 +702,23 @@ FString FBlueprintCompilerCppBackendBase::GenerateWrapperForClass(UClass* Source
 		EmitterContext.Header.AddLine(FString::Printf(TEXT("%s& GetRef__%s()"), *TypeDeclaration, *UnicodeToCPPIdentifier(Property->GetName(), false, nullptr)));
 		EmitterContext.Header.AddLine(TEXT("{"));
 		EmitterContext.Header.IncreaseIndent();
-
-		EmitterContext.Header.AddLine(TEXT("static TWeakObjectPtr<UProperty> __PropertyPtr{};"));
-		EmitterContext.Header.AddLine(TEXT("const UProperty* __Property = __PropertyPtr.Get();"));
-		EmitterContext.Header.AddLine(TEXT("if (nullptr == __Property)"));
-		EmitterContext.Header.AddLine(TEXT("{"));
-		EmitterContext.Header.IncreaseIndent();
-		EmitterContext.Header.AddLine(FString::Printf(TEXT("__Property = GetClass()->FindPropertyByName(FName(TEXT(\"%s\")));"), *Property->GetName()));
-		EmitterContext.Header.AddLine(TEXT("check(__Property);"));
-		EmitterContext.Header.AddLine(TEXT("__PropertyPtr = __Property;"));
-		EmitterContext.Header.DecreaseIndent();
-		EmitterContext.Header.AddLine(TEXT("}"));
-
+		if (bUseStaticVariables)
+		{
+			EmitterContext.Header.AddLine(TEXT("static TWeakObjectPtr<UProperty> __PropertyPtr{};"));
+			EmitterContext.Header.AddLine(TEXT("const UProperty* __Property = __PropertyPtr.Get();"));
+			EmitterContext.Header.AddLine(TEXT("if (nullptr == __Property)"));
+			EmitterContext.Header.AddLine(TEXT("{"));
+			EmitterContext.Header.IncreaseIndent();
+			EmitterContext.Header.AddLine(FString::Printf(TEXT("__Property = GetClass()->FindPropertyByName(FName(TEXT(\"%s\")));"), *Property->GetName()));
+			EmitterContext.Header.AddLine(TEXT("check(__Property);"));
+			EmitterContext.Header.AddLine(TEXT("__PropertyPtr = __Property;"));
+			EmitterContext.Header.DecreaseIndent();
+			EmitterContext.Header.AddLine(TEXT("}"));
+		}
+		else
+		{
+			EmitterContext.Header.AddLine(FString::Printf(TEXT("const UProperty* __Property = GetClass()->FindPropertyByName(FName(TEXT(\"%s\")));"), *Property->GetName()));
+		}
 		EmitterContext.Header.AddLine(FString::Printf(TEXT("return *(__Property->ContainerPtrToValuePtr<%s>(__Object));"), *TypeDeclaration));
 		EmitterContext.Header.DecreaseIndent();
 		EmitterContext.Header.AddLine(TEXT("}"));
@@ -561,8 +727,12 @@ FString FBlueprintCompilerCppBackendBase::GenerateWrapperForClass(UClass* Source
 	// FUNCTIONS:
 	for (auto Func : FunctionsToGenerate)
 	{
+		TArray<FString> FuncParameters;
+		const FString ParamNameInStructPostfix(TEXT("_"));
+		const FString FuncCppName = FEmitHelper::GetCppName(Func);
+		FString DelareFunction = FString::Printf(TEXT("void %s("), *FuncCppName);
+		FString RawParameterList;
 		{
-			FString DelcareFunction = FString::Printf(TEXT("void %s("), *FEmitHelper::GetCppName(Func));
 			bool bFirst = true;
 			for (TFieldIterator<UProperty> It(Func); It; ++It)
 			{
@@ -574,7 +744,8 @@ FString FBlueprintCompilerCppBackendBase::GenerateWrapperForClass(UClass* Source
 
 				if (!bFirst)
 				{
-					DelcareFunction += TEXT(", ");
+					DelareFunction += TEXT(", ");
+					RawParameterList += TEXT(", ");
 				}
 				else
 				{
@@ -583,49 +754,59 @@ FString FBlueprintCompilerCppBackendBase::GenerateWrapperForClass(UClass* Source
 
 				if (Property->HasAnyPropertyFlags(CPF_OutParm))
 				{
-					DelcareFunction += TEXT("/*out*/ ");
+					DelareFunction += TEXT("/*out*/ ");
 				}
 
-				DelcareFunction += EmitterContext.ExportCppDeclaration(Property, EExportedDeclaration::Parameter, EPropertyExportCPPFlags::CPPF_CustomTypeName | EPropertyExportCPPFlags::CPPF_BlueprintCppBackend);
+				DelareFunction += EmitterContext.ExportCppDeclaration(Property, EExportedDeclaration::Parameter, EPropertyExportCPPFlags::CPPF_CustomTypeName | EPropertyExportCPPFlags::CPPF_BlueprintCppBackend);
+				FString ParamAsStructMember;
+
+				{	// copied from FNativeClassHeaderGenerator::ExportEventParm
+					bool bEmitConst = Property->HasAnyPropertyFlags(CPF_ConstParm) && Property->IsA<UObjectProperty>();
+					const bool bIsConstParam = (Property->IsA(UInterfaceProperty::StaticClass()) && !Property->HasAllPropertyFlags(CPF_OutParm));
+					const bool bIsOnConstClass = (Property->IsA(UObjectProperty::StaticClass()) && ((UObjectProperty*)Property)->PropertyClass != NULL && ((UObjectProperty*)Property)->PropertyClass->HasAnyClassFlags(CLASS_Const));
+					if (bIsConstParam || bIsOnConstClass)
+					{
+						bEmitConst = false; // ExportCppDeclaration will do it for us
+					}
+
+					if (bEmitConst)
+					{
+						ParamAsStructMember = TEXT("const ");
+					}
+				}
+				ParamAsStructMember += EmitterContext.ExportCppDeclaration(Property, EExportedDeclaration::Local, EPropertyExportCPPFlags::CPPF_CustomTypeName | EPropertyExportCPPFlags::CPPF_BlueprintCppBackend, false, ParamNameInStructPostfix);
+				FuncParameters.Emplace(ParamAsStructMember);
+				RawParameterList += UnicodeToCPPIdentifier(Property->GetName(), Property->HasAnyPropertyFlags(CPF_Deprecated), TEXT("bpp__"));
 			}
-			DelcareFunction += TEXT(")");
-			EmitterContext.Header.AddLine(DelcareFunction);
 		}
-
+		DelareFunction += TEXT(")");
+		EmitterContext.Header.AddLine(DelareFunction);
 		EmitterContext.Header.AddLine(TEXT("{"));
 		EmitterContext.Header.IncreaseIndent();
 
-		EmitterContext.Header.AddLine(FString::Printf(TEXT("%s::%s __Delegate;"), *DelegatesClassName, *GenerateDelegateTypeName(Func)));
-		EmitterContext.Header.AddLine(FString::Printf(TEXT("static const FName __FunctionName(TEXT(\"%s\"));"), *Func->GetName()));
-		EmitterContext.Header.AddLine(TEXT("__Delegate.BindUFunction(__Object, __FunctionName);"));
-		EmitterContext.Header.AddLine(TEXT("check(__Delegate.IsBound());"));
+		if (bUseStaticVariables)
+		{
+			EmitterContext.Header.AddLine(FString::Printf(TEXT("static const FName __FunctionName(TEXT(\"%s\"));"), *Func->GetName()));
+		}
+		const FString FuncNameStr = bUseStaticVariables ? TEXT("__FunctionName") : FString::Printf(TEXT("FName(TEXT(\"%s\"))"), *Func->GetName());
+		EmitterContext.Header.AddLine(FString::Printf(TEXT("UFunction* __Function = __Object->FindFunctionChecked(%s);"), *FuncNameStr));
 
+		const FString FuncParametersStructName = FuncParameters.Num() ? (FuncCppName + TEXT("_Parameters")) : FString();
+		if (FuncParameters.Num())
 		{
-			FString CallFunction(TEXT("__Delegate.Execute("));
-			bool bFirst = true;
-			for (TFieldIterator<UProperty> It(Func); It; ++It)
+			EmitterContext.Header.AddLine(FString::Printf(TEXT("struct %s"), *FuncParametersStructName));
+			EmitterContext.Header.AddLine(TEXT("{"));
+			EmitterContext.Header.IncreaseIndent();
+			for (FString& Param : FuncParameters)
 			{
-				UProperty* Property = *It;
-				if (!Property || !Property->HasAnyPropertyFlags(CPF_Parm))
-				{
-					continue;
-				}
-
-				if (!bFirst)
-				{
-					CallFunction += TEXT(", ");
-				}
-				else
-				{
-					bFirst = false;
-				}
-
-				CallFunction += UnicodeToCPPIdentifier(Property->GetName(), Property->HasAnyPropertyFlags(CPF_Deprecated), TEXT("bpp__"));
+				EmitterContext.Header.AddLine(FString::Printf(TEXT("%s;"), *Param));
 			}
-			CallFunction += TEXT(");");
-			EmitterContext.Header.AddLine(CallFunction);
-		}
+			EmitterContext.Header.DecreaseIndent();
+			EmitterContext.Header.AddLine(TEXT("};"));
 
+			EmitterContext.Header.AddLine(FString::Printf(TEXT("%s __Parameters { %s };"), *FuncParametersStructName, *RawParameterList));
+		}
+		EmitterContext.Header.AddLine(FString::Printf(TEXT("__Object->ProcessEvent(__Function, %s);"), FuncParameters.Num() ? TEXT("&__Parameters") : TEXT("nullptr")));
 		EmitterContext.Header.DecreaseIndent();
 		EmitterContext.Header.AddLine(TEXT("}"));
 	}
@@ -637,7 +818,7 @@ FString FBlueprintCompilerCppBackendBase::GenerateWrapperForClass(UClass* Source
 	return EmitterContext.Header.Result;
 }
 
-void FBlueprintCompilerCppBackendBase::EmitFileBeginning(const FString& CleanName, FEmitterLocalContext& EmitterContext, bool bIncludeGeneratedH, bool bIncludeCodeHelpersInHeader)
+void FBlueprintCompilerCppBackendBase::EmitFileBeginning(const FString& CleanName, FEmitterLocalContext& EmitterContext, bool bIncludeGeneratedH, bool bIncludeCodeHelpersInHeader, bool bFullyIncludedDeclaration)
 {
 	EmitterContext.Header.AddLine(TEXT("#pragma once"));
 
@@ -724,7 +905,11 @@ void FBlueprintCompilerCppBackendBase::EmitFileBeginning(const FString& CleanNam
 		}
 	};
 
-	EmitInner(EmitterContext.Header, EmitterContext.Dependencies.IncludeInHeader, EmitterContext.Dependencies.DeclareInHeader);
+	EmitInner(EmitterContext.Header, EmitterContext.Dependencies.IncludeInHeader, bFullyIncludedDeclaration ? TSet<UField*>() : EmitterContext.Dependencies.DeclareInHeader);
+	if (bFullyIncludedDeclaration)
+	{
+		EmitInner(EmitterContext.Header, EmitterContext.Dependencies.DeclareInHeader, TSet<UField*>());
+	}
 	EmitInner(EmitterContext.Body, EmitterContext.Dependencies.IncludeInBody, TSet<UField*>());
 
 	if (bIncludeGeneratedH)
diff --git a/Engine/Source/Developer/BlueprintCompilerCppBackend/Private/BlueprintCompilerCppBackendBase.h b/Engine/Source/Developer/BlueprintCompilerCppBackend/Private/BlueprintCompilerCppBackendBase.h
--- a/Engine/Source/Developer/BlueprintCompilerCppBackend/Private/BlueprintCompilerCppBackendBase.h
+++ b/Engine/Source/Developer/BlueprintCompilerCppBackend/Private/BlueprintCompilerCppBackendBase.h
@@ -74,7 +74,7 @@ private:
 
 	static void ConstructFunctionDeclaration(FEmitterLocalContext &EmitterContext, FKismetFunctionContext &FunctionContext, TArray<UProperty*> &ArgumentList);
 
-	static FString GenerateArgList(const FEmitterLocalContext &EmitterContext, const TArray<UProperty*> &ArgumentList);
+	static FString GenerateArgList(const FEmitterLocalContext &EmitterContext, const TArray<UProperty*> &ArgumentList, bool bOnlyParamName = false);
 
 	static FString GenerateReturnType(const FEmitterLocalContext &EmitterContext, const UFunction* Function);
 
@@ -82,7 +82,7 @@ private:
 
 	void CleanBackend();
 
-	static void EmitFileBeginning(const FString& CleanName, FEmitterLocalContext& EmitterContext, bool bIncludeGeneratedH = true, bool bIncludeCodeHelpersInHeader = false);
+	static void EmitFileBeginning(const FString& CleanName, FEmitterLocalContext& EmitterContext, bool bIncludeGeneratedH = true, bool bIncludeCodeHelpersInHeader = false, bool bFullyIncludedDeclaration = false);
 
 	static void EmitStructProperties(FEmitterLocalContext& EmitterContext, UStruct* SourceClass);
 
diff --git a/Engine/Source/Developer/BlueprintCompilerCppBackend/Private/BlueprintCompilerCppBackendGatherDependencies.cpp b/Engine/Source/Developer/BlueprintCompilerCppBackend/Private/BlueprintCompilerCppBackendGatherDependencies.cpp
--- a/Engine/Source/Developer/BlueprintCompilerCppBackend/Private/BlueprintCompilerCppBackendGatherDependencies.cpp
+++ b/Engine/Source/Developer/BlueprintCompilerCppBackend/Private/BlueprintCompilerCppBackendGatherDependencies.cpp
@@ -18,13 +18,22 @@ struct FGatherConvertedClassDependenciesHelperBase : public FReferenceCollector
 
 	void FindReferences(UObject* Object)
 	{
+		UProperty* Property = Cast<UProperty>(Object);
+		if (Property && Property->HasAnyPropertyFlags(CPF_DevelopmentAssets))
+		{
+			return;
+		}
+
 		{
 			FSimpleObjectReferenceCollectorArchive CollectorArchive(Object, *this);
 			CollectorArchive.SetSerializedProperty(nullptr);
 			CollectorArchive.SetFilterEditorOnly(true);
-			Object->SerializeScriptProperties(CollectorArchive);
+			if (UClass* AsBPGC = Cast<UBlueprintGeneratedClass>(Object))
+			{
+				Object = Dependencies.FindOriginalClass(AsBPGC);
+			}
+			Object->Serialize(CollectorArchive);
 		}
-		Object->CallAddReferencedObjects(*this);
 	}
 
 	void FindReferencesForNewObject(UObject* Object)
@@ -134,8 +143,6 @@ struct FFindHeadersToInclude : public FGatherConvertedClassDependenciesHelperBas
 
 	virtual void HandleObjectReference(UObject*& InObject, const UObject* InReferencingObject, const UProperty* InReferencingProperty) override
 	{
-		const bool bIncludeMembersFromParentClass = true;
-
 		if (!InObject || InObject->IsA<UBlueprint>() || (InObject == Dependencies.GetActualStruct()))
 		{
 			return;
@@ -147,7 +154,7 @@ struct FFindHeadersToInclude : public FGatherConvertedClassDependenciesHelperBas
 			ObjAsField = InObject->GetClass();
 		}
 
-		if (ObjAsField)
+		if (ObjAsField && !ObjAsField->HasAnyFlags(RF_ClassDefaultObject))
 		{
 			if (ObjAsField->IsA<UProperty>())
 			{
@@ -164,27 +171,8 @@ struct FFindHeadersToInclude : public FGatherConvertedClassDependenciesHelperBas
 		{
 			return;
 		}
-		else if (auto ObjAsClass = Cast<UClass>(InObject))
-		{
-			if (bIncludeMembersFromParentClass && (ObjAsClass != UObject::StaticClass()) && Dependencies.GetActualStruct()->IsChildOf(ObjAsClass))
-			{
-				FindReferencesForNewObject(InObject);
-			}
-			if (ObjAsClass->HasAnyClassFlags(CLASS_Native))
-			{
-				return;
-			}
-		}
-		else if (InObject->IsA<UScriptStruct>())
-		{
-			return;
-		}
 
-		if ((InObject->GetOutermost() == Dependencies.GetActualStruct()->GetOutermost()) || 
-			(bIncludeMembersFromParentClass && ObjAsField && Dependencies.GetActualStruct()->IsChildOf(ObjAsField->GetOwnerClass())))
-		{
-			FindReferencesForNewObject(InObject);
-		}
+		FindReferencesForNewObject(InObject);
 	}
 };
 
@@ -218,7 +206,7 @@ bool FGatherConvertedClassDependencies::WillClassBeConverted(const UBlueprintGen
 		const UClass* ClassToCheck = FindOriginalClass(InClass);
 
 		IBlueprintCompilerCppBackendModule& BackEndModule = (IBlueprintCompilerCppBackendModule&)IBlueprintCompilerCppBackendModule::Get();
-		auto& WillBeConvertedQuery = BackEndModule.OnIsTargetedForConversionQuery();
+		const auto& WillBeConvertedQuery = BackEndModule.OnIsTargetedForConversionQuery();
 
 		if (WillBeConvertedQuery.IsBound())
 		{
@@ -264,15 +252,35 @@ void FGatherConvertedClassDependencies::DependenciesForHeader()
 		{
 			if (auto ObjectProperty = Cast<const UObjectPropertyBase>(Property))
 			{
-				if (ShouldIncludeHeaderFor(ObjectProperty->PropertyClass))
+				auto GetFirstNativeOrConvertedClass = [](FGatherConvertedClassDependencies& Dependencies, UClass* InClass) -> UClass*
 				{
-					DeclareInHeader.Add(ObjectProperty->PropertyClass);
-				}
+					check(InClass);
+					for (UClass* ItClass = InClass; ItClass; ItClass = ItClass->GetSuperClass())
+					{
+						auto BPGC = Cast<UBlueprintGeneratedClass>(ItClass);
+						if (ItClass->HasAnyClassFlags(CLASS_Native) || Dependencies.WillClassBeConverted(BPGC))
+						{
+							return ItClass;
+						}
+					}
+					check(false);
+					return nullptr;
+				};
+
+				DeclareInHeader.Add(GetFirstNativeOrConvertedClass(*this, ObjectProperty->PropertyClass));
 			}
 			else if (auto InterfaceProperty = Cast<const UInterfaceProperty>(Property))
 			{
 				IncludeInHeader.Add(InterfaceProperty->InterfaceClass);
 			}
+			else if (auto DelegateProperty = Cast<const UDelegateProperty>(Property))
+			{
+				IncludeInHeader.Add(DelegateProperty->SignatureFunction ? DelegateProperty->SignatureFunction->GetOwnerClass() : nullptr);
+			}
+			else if (auto MulticastDelegateProperty = Cast<const UMulticastDelegateProperty>(Property))
+			{
+				IncludeInHeader.Add(MulticastDelegateProperty->SignatureFunction ? MulticastDelegateProperty->SignatureFunction->GetOwnerClass() : nullptr);
+			}
 			else
 			{
 				HeaderReferenceFinder.FindReferences(Obj);
@@ -300,5 +308,15 @@ void FGatherConvertedClassDependencies::DependenciesForHeader()
 			IncludeInHeader.Add(CastChecked<UField>(Obj));
 		}
 	}
+
+	const UPackage* OriginalStructPackage = OriginalStruct ? OriginalStruct->GetOutermost() : nullptr;
+	for (auto Iter = IncludeInHeader.CreateIterator(); Iter; ++Iter)
+	{
+		UField* CurrentField = *Iter;
+		if (CurrentField && (CurrentField->GetOutermost() == OriginalStructPackage))
+		{
+			Iter.RemoveCurrent();
+		}
+	}
 }
 
diff --git a/Engine/Source/Developer/BlueprintCompilerCppBackend/Private/BlueprintCompilerCppBackendModule.cpp b/Engine/Source/Developer/BlueprintCompilerCppBackend/Private/BlueprintCompilerCppBackendModule.cpp
--- a/Engine/Source/Developer/BlueprintCompilerCppBackend/Private/BlueprintCompilerCppBackendModule.cpp
+++ b/Engine/Source/Developer/BlueprintCompilerCppBackend/Private/BlueprintCompilerCppBackendModule.cpp
@@ -18,11 +18,15 @@ public:
 	virtual FPCHFilenameQuery& OnPCHFilenameQuery() override;
 	virtual FIsTargetedForConversionQuery& OnIsTargetedForConversionQuery() override;
 	virtual TMap<TWeakObjectPtr<UClass>, TWeakObjectPtr<UClass> >& GetOriginalClassMap() override;
+	virtual FMarkUnconvertedBlueprintAsNecessary& OnIncludingUnconvertedBP() override;
+	virtual FIsFunctionUsedInADelegate& GetIsFunctionUsedInADelegateCallback() override;
 	//~ End IBlueprintCompilerCppBackendModule interface
 
 private: 
 	FPCHFilenameQuery PCHFilenameQuery;
 	FIsTargetedForConversionQuery IsTargetedForConversionQuery;
+	FMarkUnconvertedBlueprintAsNecessary MarkUnconvertedBlueprintAsNecessary;
+	FIsFunctionUsedInADelegate IsFunctionUsedInADelegate;
 	TMap<TWeakObjectPtr<UClass>, TWeakObjectPtr<UClass> > OriginalClassMap;
 };
 
@@ -47,9 +51,19 @@ IBlueprintCompilerCppBackendModule::FIsTargetedForConversionQuery& FBlueprintCom
 	return IsTargetedForConversionQuery;
 }
 
+IBlueprintCompilerCppBackendModule::FMarkUnconvertedBlueprintAsNecessary& FBlueprintCompilerCppBackendModule::OnIncludingUnconvertedBP()
+{
+	return MarkUnconvertedBlueprintAsNecessary;
+}
+
 TMap<TWeakObjectPtr<UClass>, TWeakObjectPtr<UClass> >& FBlueprintCompilerCppBackendModule::GetOriginalClassMap()
 {
 	return OriginalClassMap;
 }
 
+IBlueprintCompilerCppBackendModule::FIsFunctionUsedInADelegate& FBlueprintCompilerCppBackendModule::GetIsFunctionUsedInADelegateCallback()
+{
+	return IsFunctionUsedInADelegate;
+}
+
 IMPLEMENT_MODULE(FBlueprintCompilerCppBackendModule, BlueprintCompilerCppBackend)
diff --git a/Engine/Source/Developer/BlueprintCompilerCppBackend/Private/BlueprintCompilerCppBackendUtils.cpp b/Engine/Source/Developer/BlueprintCompilerCppBackend/Private/BlueprintCompilerCppBackendUtils.cpp
--- a/Engine/Source/Developer/BlueprintCompilerCppBackend/Private/BlueprintCompilerCppBackendUtils.cpp
+++ b/Engine/Source/Developer/BlueprintCompilerCppBackend/Private/BlueprintCompilerCppBackendUtils.cpp
@@ -183,7 +183,8 @@ FString FEmitterLocalContext::ExportTextItem(const UProperty* Property, const vo
 			| EPropertyExportCPPFlags::CPPF_NoRef
 			| EPropertyExportCPPFlags::CPPF_NoStaticArray
 			| EPropertyExportCPPFlags::CPPF_BlueprintCppBackend;
-		const FString TypeText = ExportCppDeclaration(ArrayProperty, EExportedDeclaration::Parameter, LocalExportCPPFlags, true);
+		const FString ConstPrefix = Property->HasMetaData(TEXT("NativeConstTemplateArg")) ? TEXT("const ") : TEXT("");
+		const FString TypeText = ExportCppDeclaration(ArrayProperty, EExportedDeclaration::Parameter, LocalExportCPPFlags, true, FString(), ConstPrefix);
 		return FString::Printf(TEXT("%s()"), *TypeText);
 	}
 
@@ -192,7 +193,7 @@ FString FEmitterLocalContext::ExportTextItem(const UProperty* Property, const vo
 	return ValueStr;
 }
 
-FString FEmitterLocalContext::ExportCppDeclaration(const UProperty* Property, EExportedDeclaration::Type DeclarationType, uint32 InExportCPPFlags, bool bSkipParameterName) const
+FString FEmitterLocalContext::ExportCppDeclaration(const UProperty* Property, EExportedDeclaration::Type DeclarationType, uint32 InExportCPPFlags, bool bSkipParameterName, const FString& NamePostfix, const FString& TypePrefix) const
 {
 	FString ActualCppType;
 	FString* ActualCppTypePtr = nullptr;
@@ -203,13 +204,13 @@ FString FEmitterLocalContext::ExportCppDeclaration(const UProperty* Property, EE
 	auto GetActualNameCPP = [&](const UObjectPropertyBase* ObjectPropertyBase, UClass* InActualClass)
 	{
 		auto BPGC = Cast<UBlueprintGeneratedClass>(InActualClass);
-		if (BPGC )
+		if (BPGC || !TypePrefix.IsEmpty())
 		{
 			const bool bIsParameter = (DeclarationType == EExportedDeclaration::Parameter) || (DeclarationType == EExportedDeclaration::MacroParameter);
 			const uint32 LocalExportCPPFlags = ExportCPPFlags | (bIsParameter ? CPPF_ArgumentOrReturnValue : 0);
-			UClass* NativeType = GetFirstNativeOrConvertedClass(BPGC);
+			UClass* NativeType = GetFirstNativeOrConvertedClass(InActualClass);
 			check(NativeType);
-			ActualCppType = ObjectPropertyBase->GetCPPTypeCustom(&ActualExtendedType, LocalExportCPPFlags, FEmitHelper::GetCppName(NativeType));
+			ActualCppType = TypePrefix + ObjectPropertyBase->GetCPPTypeCustom(&ActualExtendedType, LocalExportCPPFlags, FEmitHelper::GetCppName(NativeType));
 			ActualCppTypePtr = &ActualCppType;
 			if (!ActualExtendedType.IsEmpty())
 			{
@@ -251,31 +252,8 @@ FString FEmitterLocalContext::ExportCppDeclaration(const UProperty* Property, EE
 			ActualCppTypePtr = &ActualCppType;
 		}
 	}
-	else if (const UByteProperty* ByteProperty = Cast<const UByteProperty>(Property))
-	{
-		if (ByteProperty->Enum != nullptr)
-		{
-			const bool bIsEnumClass  = (ByteProperty->Enum->GetCppForm() == UEnum::ECppForm::EnumClass);
-			const bool bIsNativeEnum = (ByteProperty->Enum->GetClass()   == UEnum::StaticClass()); // cannot use RF_Native flag, because in UHT the flag is not set
 
-			// unfortunately, we cannot tell if the native enum is typed or not 
-			// (uint8 or not); since UHT allows untyped enums to be exposed to
-			// Blueprints, but errors when they're used for a native properties, 
-			// then we need to wrap all native class enums in TEnumAsByte<> 
-			// (to take preventative measures against non uint8 enums)
-			if (bIsNativeEnum && bIsEnumClass)
-			{
-				const FString EnumAsBytePrefix = TEXT("TEnumAsByte<");
-
-				ActualCppType = ByteProperty->GetCPPType(/*ExtendedTypeText =*/nullptr, InExportCPPFlags | CPPF_BlueprintCppBackend);
-				if (!ActualCppType.StartsWith(EnumAsBytePrefix))
-				{
-					ActualCppType = FString::Printf(TEXT("%s %s >"), *EnumAsBytePrefix, *ActualCppType);
-				}
-				ActualCppTypePtr = &ActualCppType;
-			}
-		}
-	}
+	// TODO: TypePrefix for other properties
 
 	if (ArrayProperty)
 	{
@@ -300,7 +278,7 @@ FString FEmitterLocalContext::ExportCppDeclaration(const UProperty* Property, EE
 	}
 
 	FStringOutputDevice Out;
-	const FString ActualNativeName = bSkipParameterName ? FString() : FEmitHelper::GetCppName(Property);
+	const FString ActualNativeName = bSkipParameterName ? FString() : (FEmitHelper::GetCppName(Property) + NamePostfix);
 	Property->ExportCppDeclaration(Out, DeclarationType, nullptr, InExportCPPFlags, bSkipParameterName, ActualCppTypePtr, ActualExtendedTypePtr, &ActualNativeName);
 	return FString(Out);
 
@@ -342,14 +320,7 @@ FString FEmitHelper::GetCppName(const UField* Field, bool bUInterface)
 				const bool bIsUberGraphVariable = Owner->IsA<UBlueprintGeneratedClass>() && AsProperty->HasAllPropertyFlags(CPF_Transient | CPF_DuplicateTransient);
 				if (bIsUberGraphVariable)
 				{
-					int32 InheritenceLevel = 0;
-
-					const UStruct* StructIt = Owner->GetSuperStruct();
-					while ((StructIt != nullptr) && !StructIt->IsNative())
-					{
-						++InheritenceLevel;
-						StructIt = StructIt->GetSuperStruct();
-					}
+					int32 InheritenceLevel = GetInheritenceLevel(Owner);
 					VarPrefix = FString::Printf(TEXT("b%dl__"), InheritenceLevel);
 				}
 				else
@@ -370,6 +341,23 @@ FString FEmitHelper::GetCppName(const UField* Field, bool bUInterface)
 	return Field->GetName();
 }
 
+int32 FEmitHelper::GetInheritenceLevel(const UStruct* Struct)
+{
+	const UStruct* StructIt = Struct ? Struct->GetSuperStruct() : nullptr;
+	int32 InheritenceLevel = 0;
+	while ((StructIt != nullptr) && !StructIt->IsNative())
+	{
+		++InheritenceLevel;
+		StructIt = StructIt->GetSuperStruct();
+	}
+	return InheritenceLevel;
+}
+
+bool FEmitHelper::PropertyForConstCast(const UProperty* Property)
+{
+	return Property && !Property->IsA<UStructProperty>() && Property->HasAnyPropertyFlags(CPF_ConstParm) && !Property->HasAnyPropertyFlags(CPF_OutParm);
+}
+
 void FEmitHelper::ArrayToString(const TArray<FString>& Array, FString& OutString, const TCHAR* Separator)
 {
 	if (Array.Num())
@@ -828,11 +816,11 @@ FString FEmitHelper::LiteralTerm(FEmitterLocalContext& EmitterContext, const FEd
 
 	if (UEdGraphSchema_K2::PC_String == Type.PinCategory)
 	{
-		return FString::Printf(TEXT("TEXT(\"%s\")"), *CustomValue.ReplaceCharWithEscapedChar());
+		return FString::Printf(TEXT("FString(%s)"), *UStrProperty::ExportCppHardcodedText(CustomValue, EmitterContext.DefaultTarget->Indent));
 	}
 	else if (UEdGraphSchema_K2::PC_Text == Type.PinCategory)
 	{
-		return FString::Printf(TEXT("FText::FromString(TEXT(\"%s\"))"), *CustomValue.ReplaceCharWithEscapedChar());
+		return FString::Printf(TEXT("FText::FromString(%s)"), *UStrProperty::ExportCppHardcodedText(CustomValue, EmitterContext.DefaultTarget->Indent));
 	}
 	else if (UEdGraphSchema_K2::PC_Float == Type.PinCategory)
 	{
@@ -925,14 +913,15 @@ FString FEmitHelper::LiteralTerm(FEmitterLocalContext& EmitterContext, const FEd
 			//@todo:  This needs to be more robust, since import text isn't really proper for struct construction.
 			const bool bEmptyCustomValue = CustomValue.IsEmpty() || (CustomValue == TEXT("()"));
 			const FString StructName = *FEmitHelper::GetCppName(StructType);
+			const FString LocalStructNativeName = EmitterContext.GenerateUniqueLocalName();
 			if (bEmptyCustomValue)
 			{
-				return StructName + (StructType->IsA<UUserDefinedStruct>() ? TEXT("::GetDefaultValue()") : TEXT("{}"));
+				// The local variable is created to fix: "fatal error C1001: An internal error has occurred in the compiler."
+				EmitterContext.AddLine(FString::Printf(TEXT("auto %s = %s%s;"), *LocalStructNativeName, *StructName, StructType->IsA<UUserDefinedStruct>() ? TEXT("::GetDefaultValue()") : TEXT("{}")));
 			}
-
-			const FString LocalStructNativeName = EmitterContext.GenerateUniqueLocalName();
-			EmitterContext.AddLine(FString::Printf(TEXT("auto %s = %s{};"), *LocalStructNativeName, *StructName)); // TODO: ?? should "::GetDefaultValue()" be called here?
+			else
 			{
+				EmitterContext.AddLine(FString::Printf(TEXT("auto %s = %s{};"), *LocalStructNativeName, *StructName)); // TODO: ?? should "::GetDefaultValue()" be called here?
 				FStructOnScope StructOnScope(StructType);
 
 				class FImportTextErrorContext : public FStringOutputDevice
@@ -956,7 +945,7 @@ FString FEmitHelper::LiteralTerm(FEmitterLocalContext& EmitterContext, const FEd
 				const auto EndOfParsedBuff = UStructProperty::ImportText_Static(StructType, TEXT("FEmitHelper::LiteralTerm"), *CustomValue, StructOnScope.GetStructMemory(), 0, nullptr, &ImportError);
 				if (!EndOfParsedBuff || ImportError.NumErrors)
 				{
-					UE_LOG(LogK2Compiler, Error, TEXT("FEmitHelper::LiteralTerm cannot parse struct \"%s\" error: %s"), *CustomValue, *ImportError);
+					UE_LOG(LogK2Compiler, Error, TEXT("FEmitHelper::LiteralTerm cannot parse struct \"%s\" error: %s class: %s"), *CustomValue, *ImportError, *GetPathNameSafe(EmitterContext.GetCurrentlyGeneratedClass()));
 				}
 
 				for (auto LocalProperty : TFieldRange<const UProperty>(StructType))
@@ -988,11 +977,6 @@ FString FEmitHelper::LiteralTerm(FEmitterLocalContext& EmitterContext, const FEd
 	{
 		if (LiteralObject)
 		{
-			const FString MappedObject = EmitterContext.FindGloballyMappedObject(LiteralObject);
-			if (!MappedObject.IsEmpty())
-			{
-				return MappedObject;
-			}
 			return FString::Printf(TEXT("FStringAssetReference(TEXT(\"%s\"))"), *(LiteralObject->GetPathName().ReplaceCharWithEscapedChar()));
 		}
 		return FString(TEXT("nullptr"));
@@ -1008,12 +992,8 @@ FString FEmitHelper::LiteralTerm(FEmitterLocalContext& EmitterContext, const FEd
 			{
 				return MappedObject;
 			}
-			return FString(TEXT("nullptr"));
-		}
-		else
-		{
-			return FString(TEXT("nullptr"));
 		}
+		return FString(TEXT("nullptr"));
 	}
 	else if (UEdGraphSchema_K2::PC_Interface == Type.PinCategory)
 	{
@@ -1058,8 +1038,11 @@ FString FEmitHelper::PinTypeToNativeType(const FEdGraphPinType& Type)
 		{
 			if (UEnum* Enum = Cast<UEnum>(InType.PinSubCategoryObject.Get()))
 			{
+				const bool bEnumClassForm = Enum->GetCppForm() == UEnum::ECppForm::EnumClass;
+				const bool bNonNativeEnum = Enum->GetClass() != UEnum::StaticClass();
 				FString FullyQualifiedEnumName = (!Enum->CppType.IsEmpty()) ? Enum->CppType : FEmitHelper::GetCppName(Enum);
-				return FString::Printf(TEXT("TEnumAsByte<%s>"), *FullyQualifiedEnumName);
+				// TODO: sometimes we need unwrapped type for enums without size specified. For example when native function has a raw ref param.
+				return (bEnumClassForm || bNonNativeEnum) ? FullyQualifiedEnumName : FString::Printf(TEXT("TEnumAsByte<%s>"), *FullyQualifiedEnumName);
 			}
 			return TEXT("uint8");
 		}
@@ -1239,10 +1222,10 @@ bool FEmitHelper::GenerateAutomaticCast(FEmitterLocalContext& EmitterContext, co
 			OutCastEnd = TEXT(")");
 			return true;
 		}
-		else if (RType.bIsArray && LClass && RClass && (LClass->IsChildOf(RClass) || RClass->IsChildOf(LClass)))
+		else if (RType.bIsArray && LClass && RClass && (LClass->IsChildOf(RClass) || RClass->IsChildOf(LClass)) && (LClass != RClass))
 		{
-			OutCastBegin = FString::Printf(TEXT("TArrayCaster<%s>("), *FEmitHelper::GetCppName(RClass));
-			OutCastEnd = FString::Printf(TEXT(").Get<%s>()"), *FEmitHelper::GetCppName(LClass));
+			OutCastBegin = FString::Printf(TEXT("TArrayCaster<%s*>("), *FEmitHelper::GetCppName(RClass));
+			OutCastEnd = FString::Printf(TEXT(").Get<%s*>()"), *FEmitHelper::GetCppName(LClass));
 			return true;
 		}
 	}
@@ -1266,31 +1249,44 @@ FString FEmitHelper::ReplaceConvertedMetaData(UObject* Obj)
 FString FEmitHelper::GenerateGetPropertyByName(FEmitterLocalContext& EmitterContext, const UProperty* Property)
 {
 	check(Property);
-
-	const FString PropertyWeakPtrName = EmitterContext.GenerateUniqueLocalName();
-	EmitterContext.AddLine(FString::Printf(TEXT("static TWeakObjectPtr<UProperty> %s{};"), *PropertyWeakPtrName));
-
 	const FString PropertyPtrName = EmitterContext.GenerateUniqueLocalName();
-	EmitterContext.AddLine(FString::Printf(TEXT("const UProperty* %s = %s.Get();"), *PropertyPtrName, *PropertyWeakPtrName));
-	EmitterContext.AddLine(FString::Printf(TEXT("if (nullptr == %s)"), *PropertyPtrName));
-	EmitterContext.AddLine(TEXT("{"));
-	EmitterContext.IncreaseIndent();
-
-	const FString PropertyOwnerStruct = EmitterContext.FindGloballyMappedObject(Property->GetOwnerStruct(), UStruct::StaticClass());
-	EmitterContext.AddLine(FString::Printf(TEXT("%s = (%s)->FindPropertyByName(FName(TEXT(\"%s\")));")
-		, *PropertyPtrName
-		, *PropertyOwnerStruct
-		, *Property->GetName()));
-	EmitterContext.AddLine(FString::Printf(TEXT("check(%s);"), *PropertyPtrName));
-	EmitterContext.AddLine(FString::Printf(TEXT("%s = %s;"), *PropertyWeakPtrName, *PropertyPtrName));
-	EmitterContext.DecreaseIndent();
-	EmitterContext.AddLine(TEXT("}"));
 
+	static const FBoolConfigValueHelper UseStaticVariables(TEXT("BlueprintNativizationSettings"), TEXT("bUseStaticVariablesInClasses"));
+	const bool bUseStaticVariables = UseStaticVariables;
+	if (bUseStaticVariables)
+	{
+		const FString PropertyWeakPtrName = EmitterContext.GenerateUniqueLocalName();
+		EmitterContext.AddLine(FString::Printf(TEXT("static TWeakObjectPtr<UProperty> %s{};"), *PropertyWeakPtrName));
+
+		EmitterContext.AddLine(FString::Printf(TEXT("const UProperty* %s = %s.Get();"), *PropertyPtrName, *PropertyWeakPtrName));
+		EmitterContext.AddLine(FString::Printf(TEXT("if (nullptr == %s)"), *PropertyPtrName));
+		EmitterContext.AddLine(TEXT("{"));
+		EmitterContext.IncreaseIndent();
+
+		const FString PropertyOwnerStruct = EmitterContext.FindGloballyMappedObject(Property->GetOwnerStruct(), UStruct::StaticClass());
+		EmitterContext.AddLine(FString::Printf(TEXT("%s = (%s)->FindPropertyByName(FName(TEXT(\"%s\")));")
+			, *PropertyPtrName
+			, *PropertyOwnerStruct
+			, *Property->GetName()));
+		EmitterContext.AddLine(FString::Printf(TEXT("check(%s);"), *PropertyPtrName));
+		EmitterContext.AddLine(FString::Printf(TEXT("%s = %s;"), *PropertyWeakPtrName, *PropertyPtrName));
+		EmitterContext.DecreaseIndent();
+		EmitterContext.AddLine(TEXT("}"));
+	}
+	else 
+	{
+		const FString PropertyOwnerStruct = EmitterContext.FindGloballyMappedObject(Property->GetOwnerStruct(), UStruct::StaticClass());
+		EmitterContext.AddLine(FString::Printf(TEXT("const UProperty* %s = (%s)->FindPropertyByName(FName(TEXT(\"%s\")));")
+			, *PropertyPtrName
+			, *PropertyOwnerStruct
+			, *Property->GetName()));
+		EmitterContext.AddLine(FString::Printf(TEXT("check(%s);"), *PropertyPtrName));
+	}
 	return PropertyPtrName;
 }
 
 FString FEmitHelper::AccessInaccessibleProperty(FEmitterLocalContext& EmitterContext, const UProperty* Property
-	, const FString& ContextStr, const FString& ContextAdressOp, int32 StaticArrayIdx)
+	, const FString& ContextStr, const FString& ContextAdressOp, int32 StaticArrayIdx, bool bGetter)
 {
 	check(Property);
 	const FString PropertyLocalName = GenerateGetPropertyByName(EmitterContext, Property);
@@ -1298,13 +1294,14 @@ FString FEmitHelper::AccessInaccessibleProperty(FEmitterLocalContext& EmitterCon
 		| EPropertyExportCPPFlags::CPPF_NoConst | EPropertyExportCPPFlags::CPPF_NoRef | EPropertyExportCPPFlags::CPPF_NoStaticArray
 		| EPropertyExportCPPFlags::CPPF_BlueprintCppBackend;
 	const FString TypeDeclaration = EmitterContext.ExportCppDeclaration(Property, EExportedDeclaration::Parameter, CppTemplateTypeFlags, true);
-
-	return FString::Printf(TEXT("(*(%s->ContainerPtrToValuePtr<%s>(%s(%s), %d)))")
+	const FString Result = FString::Printf(TEXT("(*(%s->ContainerPtrToValuePtr<%s>(%s(%s), %d)))")
 		, *PropertyLocalName
 		, *TypeDeclaration
 		, *ContextAdressOp
 		, *ContextStr
 		, StaticArrayIdx);
+	const UBoolProperty* BoolProperty = Cast<UBoolProperty>(Property);
+	return (bGetter && BoolProperty && !BoolProperty->IsNativeBool()) ? FString::Printf(TEXT("(0 != (%s))"), *Result) : Result;
 }
 
 namespace HelperWithoutEditorOnlyMembers
diff --git a/Engine/Source/Developer/BlueprintCompilerCppBackend/Private/BlueprintCompilerCppBackendUtils.h b/Engine/Source/Developer/BlueprintCompilerCppBackend/Private/BlueprintCompilerCppBackendUtils.h
--- a/Engine/Source/Developer/BlueprintCompilerCppBackend/Private/BlueprintCompilerCppBackendUtils.h
+++ b/Engine/Source/Developer/BlueprintCompilerCppBackend/Private/BlueprintCompilerCppBackendUtils.h
@@ -157,7 +157,7 @@ public:
 	FString FindGloballyMappedObject(const UObject* Object, const UClass* ExpectedClass = nullptr, bool bLoadIfNotFound = false, bool bTryUsedAssetsList = true);
 
 	// Functions needed for Unconverted classes
-	FString ExportCppDeclaration(const UProperty* Property, EExportedDeclaration::Type DeclarationType, uint32 AdditionalExportCPPFlags, bool bSkipParameterName = false) const;
+	FString ExportCppDeclaration(const UProperty* Property, EExportedDeclaration::Type DeclarationType, uint32 AdditionalExportCPPFlags, bool bSkipParameterName = false, const FString& NamePostfix = FString(), const FString& TypePrefix = FString()) const;
 	FString ExportTextItem(const UProperty* Property, const void* PropertyValue) const;
 
 	// AS FCodeText
@@ -185,6 +185,11 @@ struct FEmitHelper
 	// bUInterface - use interface with "U" prefix, by default there is "I" prefix
 	static FString GetCppName(const UField* Field, bool bUInterface = false);
 
+	// returns an unique number for a structure in structures hierarchy
+	static int32 GetInheritenceLevel(const UStruct* Struct);
+
+	static bool PropertyForConstCast(const UProperty* Property);
+
 	static void ArrayToString(const TArray<FString>& Array, FString& OutString, const TCHAR* Separator);
 
 	static bool HasAllFlags(uint64 Flags, uint64 FlagsToCheck);
@@ -244,7 +249,7 @@ struct FEmitHelper
 	static FString GenerateGetPropertyByName(FEmitterLocalContext& EmitterContext, const UProperty* Property);
 
 	static FString AccessInaccessibleProperty(FEmitterLocalContext& EmitterContext, const UProperty* Property
-		, const FString& ContextStr, const FString& ContextAdressOp, int32 StaticArrayIdx = 0);
+		, const FString& ContextStr, const FString& ContextAdressOp, int32 StaticArrayIdx, bool bGetter);
 
 	// This code works properly as long, as all fields in structures are UProperties!
 	static FString AccessInaccessiblePropertyUsingOffset(FEmitterLocalContext& EmitterContext, const UProperty* Property
diff --git a/Engine/Source/Developer/BlueprintCompilerCppBackend/Private/BlueprintCompilerCppBackendValueHelper.cpp b/Engine/Source/Developer/BlueprintCompilerCppBackend/Private/BlueprintCompilerCppBackendValueHelper.cpp
--- a/Engine/Source/Developer/BlueprintCompilerCppBackend/Private/BlueprintCompilerCppBackendValueHelper.cpp
+++ b/Engine/Source/Developer/BlueprintCompilerCppBackend/Private/BlueprintCompilerCppBackendValueHelper.cpp
@@ -68,7 +68,7 @@ void FEmitDefaultValueHelper::OuterGenerate(FEmitterLocalContext& Context
 				const FString ContainerStr = (EPropertyAccessOperator::None == AccessOperator) ? TEXT("this") : OuterPath;
 				const FString GetPtrStr = bNoexportProperty
 					? FEmitHelper::AccessInaccessiblePropertyUsingOffset(Context, Property, ContainerStr, OperatorStr, ArrayIndex)
-					: FEmitHelper::AccessInaccessibleProperty(Context, Property, ContainerStr, OperatorStr, ArrayIndex);
+					: FEmitHelper::AccessInaccessibleProperty(Context, Property, ContainerStr, OperatorStr, ArrayIndex, false);
 				PathToMember = Context.GenerateUniqueLocalName();
 				Context.AddLine(FString::Printf(TEXT("auto& %s = %s;"), *PathToMember, *GetPtrStr));
 			}
diff --git a/Engine/Source/Developer/BlueprintCompilerCppBackend/Public/IBlueprintCompilerCppBackendModule.h b/Engine/Source/Developer/BlueprintCompilerCppBackend/Public/IBlueprintCompilerCppBackendModule.h
--- a/Engine/Source/Developer/BlueprintCompilerCppBackend/Public/IBlueprintCompilerCppBackendModule.h
+++ b/Engine/Source/Developer/BlueprintCompilerCppBackend/Public/IBlueprintCompilerCppBackendModule.h
@@ -45,5 +45,17 @@ public:
 	 *	Cpp is generated from a duplicate of the original Blueprint. This map stores info about the original BPGC.
 	 */
 	virtual TMap<TWeakObjectPtr<UClass>, TWeakObjectPtr<UClass> >& GetOriginalClassMap() = 0;
+
+	/**
+	 *	Provides a hook so that external modules can mark some unconverted blueprints as necessary for the generated native code.
+	 */
+	DECLARE_DELEGATE_OneParam(FMarkUnconvertedBlueprintAsNecessary, TAssetPtr<UBlueprint>);
+	virtual FMarkUnconvertedBlueprintAsNecessary& OnIncludingUnconvertedBP() = 0;
+
+	DECLARE_DELEGATE_RetVal_OneParam(bool, FIsFunctionUsedInADelegate, const UFunction*);
+	virtual FIsFunctionUsedInADelegate& GetIsFunctionUsedInADelegateCallback() = 0;
+
+	// Collect functions that are used by delegates - they must have UFUNCTION macro
+	BLUEPRINTCOMPILERCPPBACKEND_API static TArray<class UFunction*> CollectBoundFunctions(class UBlueprint* BP);
 };
 
diff --git a/Engine/Source/Developer/BlueprintNativeCodeGen/Private/BlueprintNativeCodeGenModule.cpp b/Engine/Source/Developer/BlueprintNativeCodeGen/Private/BlueprintNativeCodeGenModule.cpp
--- a/Engine/Source/Developer/BlueprintNativeCodeGen/Private/BlueprintNativeCodeGenModule.cpp
+++ b/Engine/Source/Developer/BlueprintNativeCodeGen/Private/BlueprintNativeCodeGenModule.cpp
@@ -36,6 +36,10 @@ public:
 	virtual void SaveManifest(int32 Id = -1) override;
 	virtual void MergeManifest(int32 ManifestIdentifier) override;
 	virtual void FinalizeManifest() override;
+	virtual void GenerateStubs() override;
+	virtual void GenerateFullyConvertedClasses() override;
+	virtual void MarkUnconvertedBlueprintAsNecessary(TAssetPtr<UBlueprint> BPPtr) override;
+	virtual const TMultiMap<FName, TAssetSubclassOf<UObject>>& GetFunctionsBoundToADelegate() override;
 protected:
 	virtual void Initialize(const FNativeCodeGenInitData& InitData) override;
 	virtual void InitializeForRerunDebugOnly(const TArray< TPair< FString, FString > >& CodegenTargets) override;
@@ -47,38 +51,65 @@ protected:
 	virtual UClass* FindReplacedClass(const UClass* Class) const override;
 	//~ End FScriptCookReplacmentCoordinator interface
 private:
+	void ReadConfig();
+	void FillTargetedForReplacementQuery();
+	void FillIsFunctionUsedInADelegate();
 	FBlueprintNativeCodeGenManifest& GetManifest(const TCHAR* PlatformName);
+	void GenerateSingleStub(UBlueprint* BP, const TCHAR* PlatformName);
+	void CollectBoundFunctions(UBlueprint* BP);
+	void GenerateSingleAsset(UField* ForConversion, const TCHAR* PlatformName);
 
 	TMap< FString, TUniquePtr<FBlueprintNativeCodeGenManifest> > Manifests;
 
 	TArray<FString> ExcludedAssetTypes;
-	TArray<FString> ExcludedBlueprintTypes;
+	TArray<TAssetSubclassOf<UBlueprint>> ExcludedBlueprintTypes;
+	TSet<FStringAssetReference> ExcludedAssets;
 	TArray<FString> TargetPlatformNames;
+
+	// A stub-wrapper must be generated only if the BP is really accessed/required by some other generated code.
+	TSet<TAssetPtr<UBlueprint>> StubsRequiredByGeneratedCode;
+	TSet<TAssetPtr<UBlueprint>> AllPotentialStubs;
+
+	TSet<TAssetPtr<UBlueprint>> ToGenerate;
+	TMultiMap<FName, TAssetSubclassOf<UObject>> FunctionsBoundToADelegate; // is a function could be bound to a delegate, then it must have UFUNCTION macro. So we cannot optimize it.
 };
 
-void FBlueprintNativeCodeGenModule::Initialize(const FNativeCodeGenInitData& InitData)
+void FBlueprintNativeCodeGenModule::ReadConfig()
 {
 	GConfig->GetArray(TEXT("BlueprintNativizationSettings"), TEXT("ExcludedAssetTypes"), ExcludedAssetTypes, GEditorIni);
-	GConfig->GetArray(TEXT("BlueprintNativizationSettings"), TEXT("ExcludedBlueprintTypes"), ExcludedBlueprintTypes, GEditorIni);
 
-	IBlueprintNativeCodeGenCore::Register(this);
-
-	// Each platform will need a manifest, because each platform could cook different assets:
-	for (auto& Platform : InitData.CodegenTargets)
 	{
-		const TCHAR* TargetDirectory = *Platform.Value;
-		FString OutputPath = FPaths::Combine(TargetDirectory, NativizationCookControllerImpl::DefaultPluginName);
-
-		Manifests.Add(FString(*Platform.Key), TUniquePtr<FBlueprintNativeCodeGenManifest>(new FBlueprintNativeCodeGenManifest(NativizationCookControllerImpl::DefaultPluginName, OutputPath)));
+		TArray<FString> ExcludedBlueprintTypesPath;
+		GConfig->GetArray(TEXT("BlueprintNativizationSettings"), TEXT("ExcludedBlueprintTypes"), ExcludedBlueprintTypesPath, GEditorIni);
+		for (FString& Path : ExcludedBlueprintTypesPath)
+		{
+			TAssetSubclassOf<UBlueprint> ClassPtr;
+			ClassPtr = FStringAssetReference(Path);
+			ClassPtr.LoadSynchronous();
+			ExcludedBlueprintTypes.Add(ClassPtr);
+		}
+	}
 
-		TargetPlatformNames.Add(Platform.Key);
+	TArray<FString> ExcludedAssetPaths;
+	GConfig->GetArray(TEXT("BlueprintNativizationSettings"), TEXT("ExcludedAssets"), ExcludedAssetPaths, GEditorIni);
+	for (FString& Path : ExcludedAssetPaths)
+	{
+		ExcludedAssets.Add(FStringAssetReference(Path));
 	}
+}
+
+void FBlueprintNativeCodeGenModule::MarkUnconvertedBlueprintAsNecessary(TAssetPtr<UBlueprint> BPPtr)
+{
+	StubsRequiredByGeneratedCode.Add(BPPtr);
+}
 
+void FBlueprintNativeCodeGenModule::FillTargetedForReplacementQuery()
+{
 	IBlueprintCompilerCppBackendModule& BackEndModule = (IBlueprintCompilerCppBackendModule&)IBlueprintCompilerCppBackendModule::Get();
 	auto& ConversionQueryDelegate = BackEndModule.OnIsTargetedForConversionQuery();
 	auto ShouldConvert = [](const UObject* AssetObj)
 	{
-		if (IBlueprintNativeCodeGenCore::Get())
+		if (ensure(IBlueprintNativeCodeGenCore::Get()))
 		{
 			EReplacementResult ReplacmentResult = IBlueprintNativeCodeGenCore::Get()->IsTargetedForReplacement(AssetObj);
 			return ReplacmentResult == EReplacementResult::ReplaceCompletely;
@@ -86,10 +117,101 @@ void FBlueprintNativeCodeGenModule::Initialize(const FNativeCodeGenInitData& Ini
 		return false;
 	};
 	ConversionQueryDelegate.BindStatic(ShouldConvert);
+
+	auto LocalMarkUnconvertedBlueprintAsNecessary = [](TAssetPtr<UBlueprint> BPPtr)
+	{
+		IBlueprintNativeCodeGenModule::Get().MarkUnconvertedBlueprintAsNecessary(BPPtr);
+	};
+	BackEndModule.OnIncludingUnconvertedBP().BindStatic(LocalMarkUnconvertedBlueprintAsNecessary);
+}
+
+namespace 
+{
+	void GetFieldFormPackage(UPackage* Package, UStruct*& OutStruct, UEnum*& OutEnum)
+	{
+		TArray<UObject*> Objects;
+		GetObjectsWithOuter(Package, Objects, false);
+		for (auto Entry : Objects)
+		{
+			if (Entry->HasAnyFlags(RF_Transient))
+			{
+				continue;
+			}
+
+			OutStruct = Cast<UStruct>(Entry);
+			if (OutStruct)
+			{
+				break;
+			}
+
+			OutEnum = Cast<UEnum>(Entry);
+			if (OutEnum)
+			{
+				break;
+			}
+		}
+	}
+}
+
+void FBlueprintNativeCodeGenModule::CollectBoundFunctions(UBlueprint* BP)
+{
+	TArray<UFunction*> Functions = IBlueprintCompilerCppBackendModule::CollectBoundFunctions(BP);
+	for (UFunction* Func : Functions)
+	{
+		if (Func)
+		{
+			FunctionsBoundToADelegate.AddUnique(Func->GetFName(), Func->GetOwnerClass());
+		}
+	}
+}
+
+const TMultiMap<FName, TAssetSubclassOf<UObject>>& FBlueprintNativeCodeGenModule::GetFunctionsBoundToADelegate()
+{
+	return FunctionsBoundToADelegate;
+}
+
+void FBlueprintNativeCodeGenModule::FillIsFunctionUsedInADelegate()
+{
+	IBlueprintCompilerCppBackendModule& BackEndModule = (IBlueprintCompilerCppBackendModule&)IBlueprintCompilerCppBackendModule::Get();
+
+	auto IsFunctionUsed = [](const UFunction* InFunction) -> bool
+	{
+		auto& TargetFunctionsBoundToADelegate = IBlueprintNativeCodeGenModule::Get().GetFunctionsBoundToADelegate();
+		return InFunction && (nullptr != TargetFunctionsBoundToADelegate.FindPair(InFunction->GetFName(), InFunction->GetOwnerClass()));
+	};
+
+	BackEndModule.GetIsFunctionUsedInADelegateCallback().BindStatic(IsFunctionUsed);
+}
+
+void FBlueprintNativeCodeGenModule::Initialize(const FNativeCodeGenInitData& InitData)
+{
+	ReadConfig();
+
+	IBlueprintNativeCodeGenCore::Register(this);
+
+	// Each platform will need a manifest, because each platform could cook different assets:
+	for (auto& Platform : InitData.CodegenTargets)
+	{
+		const TCHAR* TargetDirectory = *Platform.Value;
+		FString OutputPath = FPaths::Combine(TargetDirectory, NativizationCookControllerImpl::DefaultPluginName);
+
+		Manifests.Add(FString(*Platform.Key), TUniquePtr<FBlueprintNativeCodeGenManifest>(new FBlueprintNativeCodeGenManifest(NativizationCookControllerImpl::DefaultPluginName, OutputPath)));
+
+		TargetPlatformNames.Add(Platform.Key);
+	}
+
+	FillTargetedForReplacementQuery();
+
+	FillIsFunctionUsedInADelegate();
 }
 
 void FBlueprintNativeCodeGenModule::InitializeForRerunDebugOnly(const TArray< TPair< FString, FString > >& CodegenTargets)
 {
+	ReadConfig();
+	IBlueprintNativeCodeGenCore::Register(this);
+	FillTargetedForReplacementQuery();
+	FillIsFunctionUsedInADelegate();
+
 	for (const auto& Platform : CodegenTargets)
 	{
 		// load the old manifest:
@@ -103,7 +225,13 @@ void FBlueprintNativeCodeGenModule::InitializeForRerunDebugOnly(const TArray< TP
 			// load the package:
 			UPackage* Package = LoadPackage(nullptr, *ConversionTarget.Value.TargetObjPath, LOAD_None);
 
-			// reconvert it:
+			if (!Package)
+			{
+				UE_LOG(LogBlueprintCodeGen, Error, TEXT("Unable to load the package: %s"), *ConversionTarget.Value.TargetObjPath);
+				continue;
+			}
+
+			// reconvert it
 			Convert(Package, ESavePackageResult::ReplaceCompletely, *Platform.Key);
 		}
 
@@ -113,11 +241,41 @@ void FBlueprintNativeCodeGenModule::InitializeForRerunDebugOnly(const TArray< TP
 			// load the package:
 			UPackage* Package = LoadPackage(nullptr, *ConversionTarget.Key.GetPlainNameString(), LOAD_None);
 
-			// reconvert it:
-			Convert(Package, ESavePackageResult::GenerateStub, *Platform.Key);
+			UStruct* Struct = nullptr;
+			UEnum* Enum = nullptr;
+			GetFieldFormPackage(Package, Struct, Enum);
+			UBlueprint* BP = Cast<UBlueprint>(CastChecked<UClass>(Struct)->ClassGeneratedBy);
+			if (ensure(BP))
+			{
+				CollectBoundFunctions(BP);
+				GenerateSingleStub(BP, *Platform.Key);
+			}
+		}
+
+		for (TAssetPtr<UBlueprint>& BPPtr : ToGenerate)
+		{
+			UBlueprint* BP = BPPtr.LoadSynchronous();
+			if (ensure(BP))
+			{
+				GenerateSingleAsset(BP->GeneratedClass, *Platform.Key);
+			}
 		}
 	}
+}
 
+void FBlueprintNativeCodeGenModule::GenerateFullyConvertedClasses()
+{
+	for (TAssetPtr<UBlueprint>& BPPtr : ToGenerate)
+	{
+		UBlueprint* BP = BPPtr.LoadSynchronous();
+		if (ensure(BP))
+		{
+			for (const FString& PlatformName : TargetPlatformNames)
+			{
+				GenerateSingleAsset(BP->GeneratedClass, *PlatformName);
+			}
+		}
+	}
 }
 
 FBlueprintNativeCodeGenManifest& FBlueprintNativeCodeGenModule::GetManifest(const TCHAR* PlatformName)
@@ -128,47 +286,40 @@ FBlueprintNativeCodeGenManifest& FBlueprintNativeCodeGenModule::GetManifest(cons
 	return **Result;
 }
 
-void FBlueprintNativeCodeGenModule::Convert(UPackage* Package, ESavePackageResult CookResult, const TCHAR* PlatformName)
+void FBlueprintNativeCodeGenModule::GenerateSingleStub(UBlueprint* BP, const TCHAR* PlatformName)
 {
-	if (CookResult != ESavePackageResult::ReplaceCompletely && CookResult != ESavePackageResult::GenerateStub)
+	UClass* Class = BP ? BP->GeneratedClass : nullptr;
+	if (!ensure(Class))
 	{
-		// nothing to convert
 		return;
 	}
 
-	// Find the struct/enum to convert:
-	UStruct* Struct = nullptr;
-	UEnum* Enum = nullptr;
-	TArray<UObject*> Objects;
-	GetObjectsWithOuter(Package, Objects, false);
-	for (auto Entry : Objects)
-	{
-		if (Entry->HasAnyFlags(RF_Transient))
-		{
-			continue;
-		}
-
-		Struct = Cast<UStruct>(Entry);
-		if (Struct)
-		{
-			break;
-		}
+	// no PCHFilename should be necessary
+	const IAssetRegistry& Registry = FModuleManager::LoadModuleChecked<FAssetRegistryModule>("AssetRegistry").Get();
+	FAssetData AssetInfo = Registry.GetAssetByObjectPath(*Class->GetPathName());
+	FString FileContents;
+	TUniquePtr<IBlueprintCompilerCppBackend> Backend_CPP(IBlueprintCompilerCppBackendModuleInterface::Get().Create());
+	// Apparently we can only generate wrappers for classes, so any logic that results in non classes requesting
+	// wrappers will fail here:
 
-		Enum = Cast<UEnum>(Entry);
-		if (Enum)
-		{
-			break;
-		}
-	}
+	FileContents = Backend_CPP->GenerateWrapperForClass(Class);
 
-	if (Struct == nullptr && Enum == nullptr)
+	if (!FileContents.IsEmpty())
 	{
-		check(false);
-		return;
+		FFileHelper::SaveStringToFile(FileContents, *(GetManifest(PlatformName).CreateUnconvertedDependencyRecord(AssetInfo.PackageName, AssetInfo).GeneratedWrapperPath));
 	}
+	// The stub we generate still may have dependencies on other modules, so make sure the module dependencies are 
+	// still recorded so that the .build.cs is generated correctly. Without this you'll get include related errors 
+	// (or possibly linker errors) in stub headers:
+	GetManifest(PlatformName).GatherModuleDependencies(BP->GetOutermost());
+}
 
+void FBlueprintNativeCodeGenModule::GenerateSingleAsset(UField* ForConversion, const TCHAR* PlatformName)
+{
 	IBlueprintCompilerCppBackendModule& BackEndModule = (IBlueprintCompilerCppBackendModule&)IBlueprintCompilerCppBackendModule::Get();
 	auto& BackendPCHQuery = BackEndModule.OnPCHFilenameQuery();
+	const IAssetRegistry& Registry = FModuleManager::LoadModuleChecked<FAssetRegistryModule>("AssetRegistry").Get();
+	FAssetData AssetInfo = Registry.GetAssetByObjectPath(*ForConversion->GetPathName());
 
 	FBlueprintNativeCodeGenPaths TargetPaths = GetManifest(PlatformName).GetTargetPaths();
 	BackendPCHQuery.BindLambda([TargetPaths]()->FString
@@ -176,81 +327,134 @@ void FBlueprintNativeCodeGenModule::Convert(UPackage* Package, ESavePackageResul
 		return TargetPaths.RuntimePCHFilename();
 	});
 
-	const IAssetRegistry& Registry = FModuleManager::LoadModuleChecked<FAssetRegistryModule>("AssetRegistry").Get();
+	FConvertedAssetRecord& ConversionRecord = GetManifest(PlatformName).CreateConversionRecord(*ForConversion->GetPathName(), AssetInfo);
+	TSharedPtr<FString> HeaderSource(new FString());
+	TSharedPtr<FString> CppSource(new FString());
 
-	if (CookResult == ESavePackageResult::GenerateStub)
+	FBlueprintNativeCodeGenUtils::GenerateCppCode(ForConversion, HeaderSource, CppSource);
+	bool bSuccess = !HeaderSource->IsEmpty() || !CppSource->IsEmpty();
+	// Run the cpp first, because we cue off of the presence of a header for a valid conversion record (see
+	// FConvertedAssetRecord::IsValid)
+	if (!CppSource->IsEmpty())
 	{
-		FAssetData AssetInfo = Registry.GetAssetByObjectPath(*Struct->GetPathName());
-		FString FileContents;
-		TUniquePtr<IBlueprintCompilerCppBackend> Backend_CPP(IBlueprintCompilerCppBackendModuleInterface::Get().Create());
-		// Apparently we can only generate wrappers for classes, so any logic that results in non classes requesting
-		// wrappers will fail here:
-
-		FileContents = Backend_CPP->GenerateWrapperForClass(CastChecked<UClass>(Struct));
-		if (!FileContents.IsEmpty())
+		if (!FFileHelper::SaveStringToFile(*CppSource, *ConversionRecord.GeneratedCppPath))
 		{
-			FFileHelper::SaveStringToFile(FileContents, *(GetManifest(PlatformName).CreateUnconvertedDependencyRecord(AssetInfo.PackageName, AssetInfo).GeneratedWrapperPath));
+			bSuccess &= false;
+			ConversionRecord.GeneratedCppPath.Empty();
 		}
-		// The stub we generate still may have dependencies on other modules, so make sure the module dependencies are 
-		// still recorded so that the .build.cs is generated correctly. Without this you'll get include related errors 
-		// (or possibly linker errors) in stub headers:
-		GetManifest(PlatformName).GatherModuleDependencies(Package);
+		CppSource->Empty(CppSource->Len());
 	}
 	else
 	{
-		check(CookResult == ESavePackageResult::ReplaceCompletely);
-		// convert:
-		UField* ForConversion = Enum;
-		if (ForConversion == nullptr)
+		ConversionRecord.GeneratedCppPath.Empty();
+	}
+
+	if (bSuccess && !HeaderSource->IsEmpty())
+	{
+		if (!FFileHelper::SaveStringToFile(*HeaderSource, *ConversionRecord.GeneratedHeaderPath))
 		{
-			ForConversion = Struct;
+			bSuccess &= false;
+			ConversionRecord.GeneratedHeaderPath.Empty();
 		}
+		HeaderSource->Empty(HeaderSource->Len());
+	}
+	else
+	{
+		ConversionRecord.GeneratedHeaderPath.Empty();
+	}
 
-		FAssetData AssetInfo = Registry.GetAssetByObjectPath(*ForConversion->GetPathName());
-		FConvertedAssetRecord& ConversionRecord = GetManifest(PlatformName).CreateConversionRecord(*ForConversion->GetPathName(), AssetInfo);
-		TSharedPtr<FString> HeaderSource(new FString());
-		TSharedPtr<FString> CppSource(new FString());
+	check(bSuccess);
+	if (bSuccess)
+	{
+		GetManifest(PlatformName).GatherModuleDependencies(ForConversion->GetOutermost());
+	}
 
-		FBlueprintNativeCodeGenUtils::GenerateCppCode(ForConversion, HeaderSource, CppSource);
-		bool bSuccess = !HeaderSource->IsEmpty() || !CppSource->IsEmpty();
-		// Run the cpp first, because we cue off of the presence of a header for a valid conversion record (see
-		// FConvertedAssetRecord::IsValid)
-		if (!CppSource->IsEmpty())
+	BackendPCHQuery.Unbind();
+}
+
+void FBlueprintNativeCodeGenModule::GenerateStubs()
+{
+	TSet<TAssetPtr<UBlueprint>> AlreadyGenerated;
+	while (AlreadyGenerated.Num() < StubsRequiredByGeneratedCode.Num())
+	{
+		const int32 OldGeneratedNum = AlreadyGenerated.Num();
+		for (TAssetPtr<UBlueprint>& BPPtr : StubsRequiredByGeneratedCode)
 		{
-			if (!FFileHelper::SaveStringToFile(*CppSource, *ConversionRecord.GeneratedCppPath))
+			bool bAlreadyGenerated = false;
+			AlreadyGenerated.Add(BPPtr, &bAlreadyGenerated);
+			if (bAlreadyGenerated)
+			{
+				continue;
+			}
+
+			ensure(AllPotentialStubs.Contains(BPPtr));
+			for (auto& PlatformName : TargetPlatformNames)
 			{
-				bSuccess &= false;
-				ConversionRecord.GeneratedCppPath.Empty();
+				GenerateSingleStub(BPPtr.LoadSynchronous(), *PlatformName);
 			}
-			CppSource->Empty(CppSource->Len());
 		}
-		else
+
+		if (!ensure(OldGeneratedNum != AlreadyGenerated.Num()))
 		{
-			ConversionRecord.GeneratedCppPath.Empty();
+			break;
 		}
+	}
 
-		if (bSuccess && !HeaderSource->IsEmpty())
+	UE_LOG(LogBlueprintCodeGen, Log, TEXT("GenerateStubs - all unconverted bp: %d, generated wrapers: %d"), AllPotentialStubs.Num(), StubsRequiredByGeneratedCode.Num());
+}
+
+void FBlueprintNativeCodeGenModule::Convert(UPackage* Package, ESavePackageResult CookResult, const TCHAR* PlatformName)
+{
+	// Find the struct/enum to convert:
+	UStruct* Struct = nullptr;
+	UEnum* Enum = nullptr;
+	GetFieldFormPackage(Package, Struct, Enum);
+
+	// First we gather information about bound functions.
+	UClass* AsClass = Cast<UClass>(Struct);
+	UBlueprint* BP = AsClass ? Cast<UBlueprint>(AsClass->ClassGeneratedBy) : nullptr;
+	if (BP)
+	{
+		CollectBoundFunctions(BP);
+	}
+
+	if (CookResult != ESavePackageResult::ReplaceCompletely && CookResult != ESavePackageResult::GenerateStub)
+	{
+		// nothing to convert
+		return;
+	}
+
+	if (Struct == nullptr && Enum == nullptr)
+	{
+		ensure(false);
+		return;
+	}
+
+	if (CookResult == ESavePackageResult::GenerateStub)
+	{
+		if (ensure(BP))
+		{
+			ensure(!ToGenerate.Contains(BP));
+			AllPotentialStubs.Add(BP);
+		}
+	}
+	else
+	{
+		check(CookResult == ESavePackageResult::ReplaceCompletely);
+		if (AsClass)
 		{
-			if (!FFileHelper::SaveStringToFile(*HeaderSource, *ConversionRecord.GeneratedHeaderPath))
+			if (ensure(BP))
 			{
-				bSuccess &= false;
-				ConversionRecord.GeneratedHeaderPath.Empty();
+				ensure(!AllPotentialStubs.Contains(BP));
+				ToGenerate.Add(BP);
 			}
-			HeaderSource->Empty(HeaderSource->Len());
 		}
 		else
 		{
-			ConversionRecord.GeneratedHeaderPath.Empty();
-		}
-
-		check(bSuccess);
-		if (bSuccess)
-		{
-			GetManifest(PlatformName).GatherModuleDependencies(Package);
+			UField* ForConversion = Struct ? (UField*)Struct : (UField*)Enum;
+			GenerateSingleAsset(ForConversion, PlatformName);
 		}
 	}
-
-	BackendPCHQuery.Unbind();
 }
 
 void FBlueprintNativeCodeGenModule::SaveManifest(int32 Id )
@@ -339,6 +543,7 @@ EReplacementResult FBlueprintNativeCodeGenModule::IsTargetedForReplacement(const
 		return EReplacementResult::DontReplace;
 	}
 
+	EReplacementResult Result = EReplacementResult::ReplaceCompletely;
 	if (const UClass* BlueprintClass = Cast<UClass>(Struct))
 	{
 		if (UBlueprint* Blueprint = Cast<UBlueprint>(BlueprintClass->ClassGeneratedBy))
@@ -357,6 +562,18 @@ EReplacementResult FBlueprintNativeCodeGenModule::IsTargetedForReplacement(const
 					return EReplacementResult::DontReplace;
 				}
 			}
+			for (TAssetSubclassOf<UBlueprint> ExcludedBlueprintTypeAsset : ExcludedBlueprintTypes)
+			{
+				UClass* ExcludedBPClass = ExcludedBlueprintTypeAsset.Get();
+				if (!ExcludedBPClass)
+				{
+					ExcludedBPClass = ExcludedBlueprintTypeAsset.LoadSynchronous();
+				}
+				if (ExcludedBPClass && Blueprint->IsA(ExcludedBPClass))
+				{
+					Result = EReplacementResult::GenerateStub;
+				}
+			}
 		}
 	}
 
@@ -393,29 +610,25 @@ EReplacementResult FBlueprintNativeCodeGenModule::IsTargetedForReplacement(const
 
 	// check blacklists:
 	// we can't use FindObject, because we may be converting a type while saving
-	if ((Struct && ExcludedAssetTypes.Find(Struct->GetPathName()) != INDEX_NONE) ||
-		(Enum && ExcludedAssetTypes.Find(Enum->GetPathName()) != INDEX_NONE))
+	if (Enum && ExcludedAssetTypes.Find(Enum->GetPathName()) != INDEX_NONE)
 	{
-		return EReplacementResult::GenerateStub;
+		Result = EReplacementResult::GenerateStub;
 	}
 
-	EReplacementResult Result = EReplacementResult::ReplaceCompletely;
 	while (Struct)
 	{
-		// This happens because the cooker incorrectly cooks editor only packages. Specifically happens for the blackjack sample
-		// project due to a FStringAssetReference in BaseEditor.ini:
-		if (Struct->RootPackageHasAnyFlags(PKG_EditorOnly))
-		{
-			return EReplacementResult::DontReplace;
-		}
-
-		if (ExcludedBlueprintTypes.Find(Struct->GetPathName()) != INDEX_NONE)
+		if (ExcludedAssetTypes.Find(Struct->GetPathName()) != INDEX_NONE)
 		{
 			Result = EReplacementResult::GenerateStub;
 		}
 		Struct = Struct->GetSuperStruct();
 	}
 
+	if (ExcludedAssets.Contains(Object->GetOutermost()))
+	{
+		Result = EReplacementResult::GenerateStub;
+	}
+
 	return Result;
 }
 
diff --git a/Engine/Source/Developer/BlueprintNativeCodeGen/Private/BlueprintNativeCodeGenUtils.cpp b/Engine/Source/Developer/BlueprintNativeCodeGen/Private/BlueprintNativeCodeGenUtils.cpp
--- a/Engine/Source/Developer/BlueprintNativeCodeGen/Private/BlueprintNativeCodeGenUtils.cpp
+++ b/Engine/Source/Developer/BlueprintNativeCodeGen/Private/BlueprintNativeCodeGenUtils.cpp
@@ -131,6 +131,7 @@ static bool BlueprintNativeCodeGenUtilsImpl::GenerateModuleSourceFiles(const FBl
 
 	TArray<FString> PchIncludes;
 	PchIncludes.Add(EngineHeaderFile);
+	PchIncludes.Add(TEXT("GeneratedCodeHelpers.h"));
 
 	bool bSuccess = GameProjectUtils::GeneratePluginModuleHeaderFile(TargetPaths.RuntimeModuleFile(FBlueprintNativeCodeGenPaths::HFile), PchIncludes, FailureReason);
 
diff --git a/Engine/Source/Developer/BlueprintNativeCodeGen/Public/BlueprintNativeCodeGenModule.h b/Engine/Source/Developer/BlueprintNativeCodeGen/Public/BlueprintNativeCodeGenModule.h
--- a/Engine/Source/Developer/BlueprintNativeCodeGen/Public/BlueprintNativeCodeGenModule.h
+++ b/Engine/Source/Developer/BlueprintNativeCodeGen/Public/BlueprintNativeCodeGenModule.h
@@ -7,6 +7,8 @@
 #include "ModuleManager.h"	// for GetModuleChecked<>
 #include "Blueprint/BlueprintSupport.h"
 
+class UBlueprint;
+
 struct FNativeCodeGenInitData
 {
 	// This is an array of pairs, the pairs are PlatformName/PlatformTargetDirectory. These
@@ -67,6 +69,10 @@ public:
 	virtual void SaveManifest(int32 Id = -1) = 0;
 	virtual void MergeManifest(int32 ManifestIdentifier) = 0;
 	virtual void FinalizeManifest() = 0;
+	virtual void GenerateStubs() = 0;
+	virtual void GenerateFullyConvertedClasses() = 0;
+	virtual void MarkUnconvertedBlueprintAsNecessary(TAssetPtr<UBlueprint> BPPtr) = 0;
+	virtual const TMultiMap<FName, TAssetSubclassOf<UObject>>& GetFunctionsBoundToADelegate() = 0;
 protected:
 	virtual void Initialize(const FNativeCodeGenInitData& InitData) = 0;
 	virtual void InitializeForRerunDebugOnly(const TArray< TPair< FString, FString > >& CodegenTargets) = 0;
diff --git a/Engine/Source/Editor/BlueprintGraph/Classes/K2Node_Composite.h b/Engine/Source/Editor/BlueprintGraph/Classes/K2Node_Composite.h
--- a/Engine/Source/Editor/BlueprintGraph/Classes/K2Node_Composite.h
+++ b/Engine/Source/Editor/BlueprintGraph/Classes/K2Node_Composite.h
@@ -13,6 +13,10 @@ class UK2Node_Composite : public UK2Node_Tunnel
 	UPROPERTY()
 	class UEdGraph* BoundGraph;
 
+	//~ Begin UObject Interface
+	virtual void PostEditUndo() override;
+	//~ End UObject Interface
+
 	//~ Begin UEdGraphNode Interface
 	BLUEPRINTGRAPH_API virtual void AllocateDefaultPins() override;
 	BLUEPRINTGRAPH_API virtual void DestroyNode() override;
@@ -37,6 +41,10 @@ class UK2Node_Composite : public UK2Node_Tunnel
 	BLUEPRINTGRAPH_API UK2Node_Tunnel* GetEntryNode() const;
 	BLUEPRINTGRAPH_API UK2Node_Tunnel* GetExitNode() const;
 
+protected:
+	/** Fixes up the input and output sink when needed, useful after PostEditUndo which changes which graph these nodes point to */
+	void FixupInputAndOutputSink();
+
 private:
 	/** Rename the BoundGraph to a unique name
 		- this is a special case rename: RenameGraphCloseToName() assumes that we are only concurrned with the immediate Outer() scope,
diff --git a/Engine/Source/Editor/BlueprintGraph/Classes/K2Node_Copy.h b/Engine/Source/Editor/BlueprintGraph/Classes/K2Node_Copy.h
new file mode 100644
--- /dev/null
+++ b/Engine/Source/Editor/BlueprintGraph/Classes/K2Node_Copy.h
@@ -0,0 +1,42 @@
+// Copyright 1998-2016 Epic Games, Inc. All Rights Reserved.
+
+#pragma once
+#include "K2Node.h"
+#include "K2Node_Copy.generated.h"
+
+UCLASS(MinimalAPI, meta=(Keywords = "Duplicate"))
+class UK2Node_Copy : public UK2Node
+{
+	GENERATED_UCLASS_BODY()
+
+	//~ Begin UEdGraphNode Interface
+	virtual void AllocateDefaultPins() override;
+	virtual FText GetTooltipText() const override;
+	virtual FText GetNodeTitle(ENodeTitleType::Type TitleType) const override;
+	virtual void PinTypeChanged(UEdGraphPin* Pin) override;
+	virtual void PostPasteNode() override;
+	//~ End UEdGraphNode Interface
+
+	//~ Begin UK2Node Interface
+	virtual void PostReconstructNode() override;
+	virtual bool IsNodeSafeToIgnore() const override { return true; }
+	virtual bool IsNodePure() const override { return true; }
+	virtual void NotifyPinConnectionListChanged(UEdGraphPin* Pin) override;
+	virtual class FNodeHandlingFunctor* CreateNodeHandler(class FKismetCompilerContext& CompilerContext) const override;
+	virtual bool IsConnectionDisallowed(const UEdGraphPin* MyPin, const UEdGraphPin* OtherPin, FString& OutReason) const override;
+	virtual void GetMenuActions(FBlueprintActionDatabaseRegistrar& ActionRegistrar) const override;
+	virtual FText GetMenuCategory() const override;
+	virtual int32 GetNodeRefreshPriority() const override { return EBaseNodeRefreshPriority::Low_UsesDependentWildcard; }
+	virtual bool AllowSplitPins() const override { return false; }
+	//~ End UK2Node Interface
+
+	/** Get the input reference pin */
+	BLUEPRINTGRAPH_API UEdGraphPin* GetInputReferencePin() const;
+	/** Get the return value pin */
+	BLUEPRINTGRAPH_API UEdGraphPin* GetCopyResultPin() const;
+
+protected:
+	/** Propagates pin type to the between the input and output pins */
+	void PropagatePinType(FEdGraphPinType& InType);
+};
+
diff --git a/Engine/Source/Editor/BlueprintGraph/Classes/K2Node_Variable.h b/Engine/Source/Editor/BlueprintGraph/Classes/K2Node_Variable.h
--- a/Engine/Source/Editor/BlueprintGraph/Classes/K2Node_Variable.h
+++ b/Engine/Source/Editor/BlueprintGraph/Classes/K2Node_Variable.h
@@ -66,6 +66,8 @@ public:
 	virtual void AutowireNewNode(UEdGraphPin* FromPin) override;
 	virtual bool CanPasteHere(const UEdGraph* TargetGraph) const override;
 	virtual void PostPasteNode() override;
+	virtual bool IsDeprecated() const;
+	virtual FString GetDeprecationMessage() const;
 	//~ End UEdGraphNode Interface
 
 	//~ Begin K2Node Interface
diff --git a/Engine/Source/Editor/BlueprintGraph/Private/K2Node_Composite.cpp b/Engine/Source/Editor/BlueprintGraph/Private/K2Node_Composite.cpp
--- a/Engine/Source/Editor/BlueprintGraph/Private/K2Node_Composite.cpp
+++ b/Engine/Source/Editor/BlueprintGraph/Private/K2Node_Composite.cpp
@@ -126,6 +126,43 @@ void UK2Node_Composite::PostPasteNode()
 	}
 }
 
+void UK2Node_Composite::PostEditUndo()
+{
+	Super::PostEditUndo();
+	FixupInputAndOutputSink();
+}
+
+void UK2Node_Composite::FixupInputAndOutputSink()
+{
+	if (BoundGraph)
+	{
+		// Update the InputSinkNode / OutputSourceNode pointers to point to the new graph
+		for (UEdGraphNode* Node : BoundGraph->Nodes)
+		{
+			if (Node->GetClass() == UK2Node_Tunnel::StaticClass())
+			{
+				// Exactly a tunnel node, should be the entrance or exit node
+				UK2Node_Tunnel* Tunnel = CastChecked<UK2Node_Tunnel>(Node);
+
+				if (Tunnel->bCanHaveInputs && !Tunnel->bCanHaveOutputs)
+				{
+					OutputSourceNode = Tunnel;
+					Tunnel->InputSinkNode = this;
+				}
+				else if (Tunnel->bCanHaveOutputs && !Tunnel->bCanHaveInputs)
+				{
+					InputSinkNode = Tunnel;
+					Tunnel->OutputSourceNode = this;
+				}
+				else
+				{
+					ensureMsgf(false, *LOCTEXT("UnexpectedTunnelNode", "Unexpected tunnel node '%s' in cloned graph '%s' (both I/O or neither)").ToString(), *Tunnel->GetName(), *GetName());
+				}
+			}
+		}
+	}
+}
+
 FText UK2Node_Composite::GetTooltipText() const
 {
 	if (InputSinkNode != NULL)
diff --git a/Engine/Source/Editor/BlueprintGraph/Private/K2Node_Copy.cpp b/Engine/Source/Editor/BlueprintGraph/Private/K2Node_Copy.cpp
new file mode 100644
--- /dev/null
+++ b/Engine/Source/Editor/BlueprintGraph/Private/K2Node_Copy.cpp
@@ -0,0 +1,283 @@
+// Copyright 1998-2016 Epic Games, Inc. All Rights Reserved.
+
+#include "BlueprintGraphPrivatePCH.h"
+#include "../../../Runtime/Engine/Classes/Kismet/KismetMathLibrary.h"
+#include "../../../Runtime/Engine/Classes/Kismet/KismetSystemLibrary.h"
+#include "KismetCompiler.h"
+#include "BlueprintNodeSpawner.h"
+#include "EditorCategoryUtils.h"
+#include "BlueprintActionDatabaseRegistrar.h"
+
+#define LOCTEXT_NAMESPACE "K2Node_Copy"
+
+//////////////////////////////////////////////////////////////////////////
+// FKCHandler_Select
+
+class FKCHandler_Copy : public FNodeHandlingFunctor
+{
+protected:
+	TMap<UEdGraphNode*, FBPTerminal*> BoolTermMap;
+
+public:
+	FKCHandler_Copy(FKismetCompilerContext& InCompilerContext)
+		: FNodeHandlingFunctor(InCompilerContext)
+	{
+	}
+
+	virtual void RegisterNets(FKismetFunctionContext& Context, UEdGraphNode* Node) override
+	{
+		FNodeHandlingFunctor::RegisterNets(Context, Node);
+
+		// Create the net for the return value manually as it's a special case Output Direction pin
+		UK2Node_Copy* CopyNode = Cast<UK2Node_Copy>(Node);
+		UEdGraphPin* CopyResultPin = CopyNode->GetCopyResultPin();
+
+		FBPTerminal* Term = Context.CreateLocalTerminalFromPinAutoChooseScope(CopyResultPin, Context.NetNameMap->MakeValidName(CopyResultPin));
+		Context.NetMap.Add(CopyResultPin, Term);
+	}
+
+	virtual void Compile(FKismetFunctionContext& Context, UEdGraphNode* Node) override
+	{
+		// Cast the node and get all the input pins
+		UK2Node_Copy* CopyNode = Cast<UK2Node_Copy>(Node);
+
+		// Get the kismet term for the copy result
+		UEdGraphPin* CopyResultPin = CopyNode->GetCopyResultPin();
+		FBPTerminal** CopyResultTerm = Context.NetMap.Find(CopyResultPin);
+
+		// Get the kismet term from the option pin
+		UEdGraphPin* InputReferencePin = FEdGraphUtilities::GetNetFromPin(Node->Pins[0]);
+		FBPTerminal** InputReferenceTerm = Context.NetMap.Find(InputReferencePin);
+
+		// Copy the reference value into the output pin so the value will be returned as a copy.
+		FKismetCompilerUtilities::CreateObjectAssignmentStatement(Context, Node, *InputReferenceTerm, *CopyResultTerm);
+	}
+};
+
+UK2Node_Copy::UK2Node_Copy(const FObjectInitializer& ObjectInitializer)
+	: Super(ObjectInitializer)
+{
+}
+
+void UK2Node_Copy::AllocateDefaultPins()
+{
+	const UEdGraphSchema_K2* Schema = GetDefault<UEdGraphSchema_K2>();
+
+	CreatePin(EGPD_Input, Schema->PC_Wildcard, TEXT(""), NULL, false, false, Schema->PN_Item);
+	CreatePin(EGPD_Output, Schema->PC_Wildcard, TEXT(""), NULL, false, false, Schema->PN_ReturnValue);
+
+	Super::AllocateDefaultPins();
+}
+
+FText UK2Node_Copy::GetTooltipText() const
+{
+	return LOCTEXT("CopyNodeTooltip", "Outputs a copy of the value passed into it.");
+}
+
+FText UK2Node_Copy::GetNodeTitle(ENodeTitleType::Type TitleType) const
+{
+	return LOCTEXT("Copy", "Copy");
+}
+
+void UK2Node_Copy::PostReconstructNode()
+{
+	if (GetInputReferencePin()->LinkedTo.Num() > 0)
+	{
+		PropagatePinType(GetInputReferencePin()->LinkedTo[0]->PinType);
+	}
+	else if (GetCopyResultPin()->LinkedTo.Num() > 0)
+	{
+		PropagatePinType(GetCopyResultPin()->LinkedTo[0]->PinType);
+	}
+}
+
+/** Determine if any pins are connected, if so make all the other pins the same type, if not, make sure pins are switched back to wildcards */
+void UK2Node_Copy::NotifyPinConnectionListChanged(UEdGraphPin* Pin)
+{
+	Super::NotifyPinConnectionListChanged(Pin);
+
+	if (Pin->ParentPin == nullptr)
+	{
+		UEdGraphPin* InputPin = GetInputReferencePin();
+		UEdGraphPin* ResultPin = GetCopyResultPin();
+
+		// If a pin connection is made, then we need to propagate the change to the other pins and validate all connections, otherwise reset them to Wildcard pins
+		if (Pin->LinkedTo.Num() > 0)
+		{
+			PropagatePinType(Pin->LinkedTo[0]->PinType);
+		}
+		else if (InputPin->LinkedTo.Num() == 0 && ResultPin->LinkedTo.Num() == 0)
+		{
+			InputPin->PinType.PinCategory = UEdGraphSchema_K2::PC_Wildcard;
+			InputPin->PinType.PinSubCategory = TEXT("");
+			InputPin->PinType.PinSubCategoryObject = NULL;
+
+			ResultPin->PinType.PinCategory = UEdGraphSchema_K2::PC_Wildcard;
+			ResultPin->PinType.PinSubCategory = TEXT("");
+			ResultPin->PinType.PinSubCategoryObject = NULL;
+
+			InputPin->BreakAllPinLinks();
+			ResultPin->BreakAllPinLinks();
+		}
+	}
+}
+
+UEdGraphPin* UK2Node_Copy::GetInputReferencePin() const
+{
+	const UEdGraphSchema_K2* K2Schema = GetDefault<UEdGraphSchema_K2>();
+
+	UEdGraphPin* Pin = FindPin(K2Schema->PN_Item);
+	check(Pin != NULL);
+	return Pin;
+}
+
+UEdGraphPin* UK2Node_Copy::GetCopyResultPin() const
+{
+	const UEdGraphSchema_K2* K2Schema = GetDefault<UEdGraphSchema_K2>();
+
+	UEdGraphPin* Pin = FindPin(K2Schema->PN_ReturnValue);
+	check(Pin != NULL);
+	return Pin;
+}
+
+void UK2Node_Copy::PinTypeChanged(UEdGraphPin* Pin)
+{
+	const UEdGraphSchema_K2* Schema = GetDefault<UEdGraphSchema_K2>();
+
+	{
+		// Set the return value
+		UEdGraphPin* ReturnPin = GetCopyResultPin();
+		if (ReturnPin->PinType != Pin->PinType)
+		{
+			// Recombine the sub pins back into the ReturnPin
+			if (ReturnPin->SubPins.Num() > 0)
+			{
+				Schema->RecombinePin(ReturnPin->SubPins[0]);
+			}
+			ReturnPin->PinType = Pin->PinType;
+			Schema->SetPinDefaultValueBasedOnType(ReturnPin);
+		}
+	}
+
+	// Let the graph know to refresh
+	GetGraph()->NotifyGraphChanged();
+
+	UBlueprint* Blueprint = GetBlueprint();
+	if(!Blueprint->bBeingCompiled)
+	{
+		FBlueprintEditorUtils::MarkBlueprintAsModified(Blueprint);
+		Blueprint->BroadcastChanged();
+	}
+}
+
+void UK2Node_Copy::PostPasteNode()
+{
+	Super::PostPasteNode();
+}
+
+bool UK2Node_Copy::IsConnectionDisallowed(const UEdGraphPin* MyPin, const UEdGraphPin* OtherPin, FString& OutReason) const
+{
+	const UEdGraphSchema_K2* K2Schema = GetDefault<UEdGraphSchema_K2>();
+	if (OtherPin)
+	{
+		if (OtherPin->PinType.PinCategory == K2Schema->PC_Exec)
+		{
+			OutReason = LOCTEXT("ExecConnectionDisallowed", "Cannot connect with Exec pin.").ToString();
+			return true;
+		}
+		else if (OtherPin->PinType.PinCategory == K2Schema->PC_Object || OtherPin->PinType.PinCategory == K2Schema->PC_Class
+			|| OtherPin->PinType.PinCategory == K2Schema->PC_Asset || OtherPin->PinType.PinCategory == K2Schema->PC_AssetClass
+			|| OtherPin->PinType.PinCategory == K2Schema->PC_Interface)
+		{
+			OutReason = FText::Format(LOCTEXT("ObjectConnectionDisallowed", "Cannot connect with {0} pin."), FText::FromString(OtherPin->PinType.PinCategory)).ToString();
+			return true;
+		}
+		else if (OtherPin->PinType.bIsArray)
+		{
+			OutReason = LOCTEXT("ArrayConnectionDisallowed", "Cannot connect with array pin.").ToString();
+			return true;
+		}
+	}
+
+	return Super::IsConnectionDisallowed(MyPin, OtherPin, OutReason);
+}
+
+FNodeHandlingFunctor* UK2Node_Copy::CreateNodeHandler(FKismetCompilerContext& CompilerContext) const
+{
+	return static_cast<FNodeHandlingFunctor*>(new FKCHandler_Copy(CompilerContext));
+}
+
+void UK2Node_Copy::GetMenuActions(FBlueprintActionDatabaseRegistrar& ActionRegistrar) const
+{
+	// actions get registered under specific object-keys; the idea is that 
+	// actions might have to be updated (or deleted) if their object-key is  
+	// mutated (or removed)... here we use the node's class (so if the node 
+	// type disappears, then the action should go with it)
+	UClass* ActionKey = GetClass();
+	// to keep from needlessly instantiating a UBlueprintNodeSpawner, first   
+	// check to make sure that the registrar is looking for actions of this type
+	// (could be regenerating actions for a specific asset, and therefore the 
+	// registrar would only accept actions corresponding to that asset)
+	if (ActionRegistrar.IsOpenForRegistration(ActionKey))
+	{
+		UBlueprintNodeSpawner* NodeSpawner = UBlueprintNodeSpawner::Create(GetClass());
+		check(NodeSpawner != nullptr);
+
+		ActionRegistrar.AddBlueprintAction(ActionKey, NodeSpawner);
+	}
+}
+
+FText UK2Node_Copy::GetMenuCategory() const
+{
+	return FEditorCategoryUtils::GetCommonCategory(FCommonEditorCategory::Utilities);
+}
+
+void UK2Node_Copy::PropagatePinType(FEdGraphPinType& InType)
+{
+	UClass const* CallingContext = NULL;
+	if (UBlueprint const* Blueprint = GetBlueprint())
+	{
+		CallingContext = Blueprint->GeneratedClass;
+		if (CallingContext == NULL)
+		{
+			CallingContext = Blueprint->ParentClass;
+		}
+	}
+
+	UEdGraphPin* InputPin = GetInputReferencePin();
+	UEdGraphPin* ResultPin = GetCopyResultPin();
+
+	InputPin->PinType = InType;
+	ResultPin->PinType = InType;
+
+	InputPin->PinType.bIsArray = false;
+	InputPin->PinType.bIsReference = false;
+
+	ResultPin->PinType.bIsArray = false;
+	ResultPin->PinType.bIsReference = false;
+
+	const UEdGraphSchema_K2* Schema = GetDefault<UEdGraphSchema_K2>();
+
+	// Verify that all previous connections to this pin are still valid with the new type
+	for (TArray<UEdGraphPin*>::TIterator ConnectionIt(InputPin->LinkedTo); ConnectionIt; ++ConnectionIt)
+	{
+		UEdGraphPin* ConnectedPin = *ConnectionIt;
+		if (!Schema->ArePinsCompatible(InputPin, ConnectedPin, CallingContext))
+		{
+			InputPin->BreakLinkTo(ConnectedPin);
+		}
+	}
+
+	// Verify that all previous connections to this pin are still valid with the new type
+	for (TArray<UEdGraphPin*>::TIterator ConnectionIt(ResultPin->LinkedTo); ConnectionIt; ++ConnectionIt)
+	{
+		UEdGraphPin* ConnectedPin = *ConnectionIt;
+		if (!Schema->ArePinsCompatible(ResultPin, ConnectedPin, CallingContext))
+		{
+			ResultPin->BreakLinkTo(ConnectedPin);
+		}
+	}
+}
+
+
+#undef LOCTEXT_NAMESPACE
diff --git a/Engine/Source/Editor/BlueprintGraph/Private/K2Node_Variable.cpp b/Engine/Source/Editor/BlueprintGraph/Private/K2Node_Variable.cpp
--- a/Engine/Source/Editor/BlueprintGraph/Private/K2Node_Variable.cpp
+++ b/Engine/Source/Editor/BlueprintGraph/Private/K2Node_Variable.cpp
@@ -847,4 +847,31 @@ void UK2Node_Variable::PostPasteNode()
 	}
 }
 
+bool UK2Node_Variable::IsDeprecated() const
+{
+	if (Super::IsDeprecated() || VariableReference.IsDeprecated())
+	{
+		return true;
+	}
+
+	UProperty* VariableProperty = VariableReference.ResolveMember<UProperty>(GetBlueprintClassFromNode());
+	if (VariableProperty 
+		&& (VariableProperty->HasAllPropertyFlags(CPF_Deprecated) || VariableProperty->HasMetaData(FBlueprintMetadata::MD_DeprecationMessage)))
+	{
+		return true;
+	}
+	return false;
+}
+
+FString UK2Node_Variable::GetDeprecationMessage() const
+{
+	UProperty* VariableProperty = VariableReference.ResolveMember<UProperty>(GetBlueprintClassFromNode());
+	if (VariableProperty && VariableProperty->HasMetaData(FBlueprintMetadata::MD_DeprecationMessage))
+	{
+		return FString::Printf(TEXT("%s %s"), *LOCTEXT("PropertyDeprecated_Warning", "@@ is deprecated;").ToString(), *VariableProperty->GetMetaData(FBlueprintMetadata::MD_DeprecationMessage));
+	}
+
+	return Super::GetDeprecationMessage();
+}
+
 #undef LOCTEXT_NAMESPACE
diff --git a/Engine/Source/Editor/BlueprintGraph/Public/BlueprintGraphDefinitions.h b/Engine/Source/Editor/BlueprintGraph/Public/BlueprintGraphDefinitions.h
--- a/Engine/Source/Editor/BlueprintGraph/Public/BlueprintGraphDefinitions.h
+++ b/Engine/Source/Editor/BlueprintGraph/Public/BlueprintGraphDefinitions.h
@@ -23,6 +23,7 @@
 #include "K2Node_DoOnceMultiInput.h"
 #include "K2Node_ComponentBoundEvent.h"
 #include "K2Node_Composite.h"
+#include "K2Node_Copy.h"
 #include "K2Node_CreateDelegate.h"
 #include "K2Node_CustomEvent.h"
 #include "K2Node_DelegateSet.h"
diff --git a/Engine/Source/Editor/EditorStyle/Private/SlateEditorStyle.cpp b/Engine/Source/Editor/EditorStyle/Private/SlateEditorStyle.cpp
--- a/Engine/Source/Editor/EditorStyle/Private/SlateEditorStyle.cpp
+++ b/Engine/Source/Editor/EditorStyle/Private/SlateEditorStyle.cpp
@@ -3918,6 +3918,12 @@ void FSlateEditorStyle::FStyle::SetupGraphEditorStyles()
 		Set( "Graph.RefPin.ConnectedHovered", new IMAGE_BRUSH( "/Graph/RefPin_connected", FVector2D(11,11) ) );
 		Set( "Graph.RefPin.DisconnectedHovered", new IMAGE_BRUSH( "/Graph/RefPin_disconnected", FVector2D(11,11) ) );
 
+		Set("Graph.Pin.CopyNodePinLeft_Connected", new IMAGE_BRUSH("/Graph/CopyNodePinLeft_connected", FVector2D(12, 24)));
+		Set("Graph.Pin.CopyNodePinLeft_Disconnected", new IMAGE_BRUSH("/Graph/CopyNodePinLeft_disconnected", FVector2D(12, 24)));
+
+		Set("Graph.Pin.CopyNodePinRight_Connected", new IMAGE_BRUSH("/Graph/CopyNodePinRight_connected", FVector2D(12, 24)));
+		Set("Graph.Pin.CopyNodePinRight_Disconnected", new IMAGE_BRUSH("/Graph/CopyNodePinRight_disconnected", FVector2D(12, 24)));
+
 		/** Variant A Pin Styles */
 		Set( "Graph.Pin.Connected_VarA", new IMAGE_BRUSH( "/Graph/Pin_connected_VarA", FVector2D(15,11)) );
 		Set( "Graph.Pin.Disconnected_VarA", new IMAGE_BRUSH( "/Graph/Pin_disconnected_VarA", FVector2D(15,11)) );
diff --git a/Engine/Source/Editor/GraphEditor/Private/KismetNodes/SGraphNodeK2Copy.cpp b/Engine/Source/Editor/GraphEditor/Private/KismetNodes/SGraphNodeK2Copy.cpp
new file mode 100644
--- /dev/null
+++ b/Engine/Source/Editor/GraphEditor/Private/KismetNodes/SGraphNodeK2Copy.cpp
@@ -0,0 +1,281 @@
+// Copyright 1998-2016 Epic Games, Inc. All Rights Reserved.
+
+#include "GraphEditorCommon.h"
+#include "Engine/Breakpoint.h"
+#include "SGraphNodeK2Base.h"
+#include "SGraphNodeK2Copy.h"
+#include "Editor/UnrealEd/Public/Kismet2/KismetDebugUtilities.h"
+#include "Editor/UnrealEd/Public/Kismet2/BlueprintEditorUtils.h"
+#include "KismetNodeInfoContext.h"
+#include "IDocumentation.h"
+#include "TutorialMetaData.h"
+
+// Pin stuff
+#include "SLevelOfDetailBranchNode.h"
+
+#define LOCTEXT_NAMESPACE "SGraphNodeK2Copy"
+
+class SCopyNodeGraphPin : public SGraphPin
+{
+public:
+
+	SLATE_BEGIN_ARGS(SCopyNodeGraphPin)
+		: _PinLabelStyle(NAME_DefaultPinLabelStyle)
+		, _UsePinColorForText(false)
+		, _SideToSideMargin(0.0f)
+	{}
+		SLATE_ARGUMENT(FName, PinLabelStyle)
+		SLATE_ARGUMENT(bool, UsePinColorForText)
+		SLATE_ARGUMENT(float, SideToSideMargin)
+	SLATE_END_ARGS()
+
+	void Construct(const FArguments& InArgs, UEdGraphPin* InPin)
+	{
+		bUsePinColorForText = InArgs._UsePinColorForText;
+		this->SetCursor(EMouseCursor::Default);
+
+		Visibility = TAttribute<EVisibility>(this, &SCopyNodeGraphPin::GetPinVisiblity);
+
+		GraphPinObj = InPin;
+		check(GraphPinObj != NULL);
+
+		const UEdGraphSchema* Schema = GraphPinObj->GetSchema();
+		checkf(
+			Schema,
+			TEXT("Missing schema for pin: %s with outer: %s of type %s"),
+			*(GraphPinObj->GetName()),
+			GraphPinObj->GetOuter() ? *(GraphPinObj->GetOuter()->GetName()) : TEXT("NULL OUTER"),
+			GraphPinObj->GetOuter() ? *(GraphPinObj->GetOuter()->GetClass()->GetName()) : TEXT("NULL OUTER")
+			);
+
+		const bool bCanConnectToPin = !GraphPinObj->bNotConnectable;
+		const bool bIsInput = (GetDirection() == EGPD_Input);
+		
+		// Create the pin icon widget
+		TSharedRef<SWidget> ActualPinWidget =
+			SAssignNew(PinImage, SImage)
+			.Image(this, &SCopyNodeGraphPin::GetPinIcon)
+			.IsEnabled(bCanConnectToPin)
+			.ColorAndOpacity(this, &SCopyNodeGraphPin::GetPinColor)
+			.OnMouseButtonDown(this, &SCopyNodeGraphPin::OnPinMouseDown)
+			.Cursor(this, &SCopyNodeGraphPin::GetPinCursor);
+
+		// Create the pin indicator widget (used for watched values)
+		static const FName NAME_NoBorder("NoBorder");
+		TSharedRef<SWidget> PinStatusIndicator =
+			SNew(SButton)
+			.ButtonStyle(FEditorStyle::Get(), NAME_NoBorder)
+			.Visibility(this, &SCopyNodeGraphPin::GetPinStatusIconVisibility)
+			.ContentPadding(0)
+			.OnClicked(this, &SCopyNodeGraphPin::ClickedOnPinStatusIcon)
+			[
+				SNew(SImage)
+				.Image(this, &SCopyNodeGraphPin::GetPinStatusIcon)
+			];
+
+		// Create the widget used for the pin body (status indicator, label, and value)
+		TSharedRef<SWrapBox> LabelAndValue =
+			SNew(SWrapBox)
+			.PreferredWidth(150.f);
+
+		LabelAndValue->AddSlot()
+			.VAlign(VAlign_Center)
+			[
+				PinStatusIndicator
+			];
+
+		TSharedPtr<SWidget> PinContent;
+		if (bIsInput)
+		{
+			// Input pin
+			PinContent = SNew(SHorizontalBox)
+				+ SHorizontalBox::Slot()
+				.AutoWidth()
+				.VAlign(VAlign_Center)
+				.Padding(0, 0, InArgs._SideToSideMargin, 0)
+				[
+					ActualPinWidget
+				]
+			+ SHorizontalBox::Slot()
+				.AutoWidth()
+				.VAlign(VAlign_Center)
+				[
+					LabelAndValue
+				];
+		}
+		else
+		{
+			// Output pin
+			PinContent = SNew(SHorizontalBox)
+				+ SHorizontalBox::Slot()
+				.AutoWidth()
+				.VAlign(VAlign_Center)
+				[
+					LabelAndValue
+				]
+			+ SHorizontalBox::Slot()
+				.AutoWidth()
+				.VAlign(VAlign_Center)
+				.Padding(InArgs._SideToSideMargin, 0, 0, 0)
+				[
+					ActualPinWidget
+				];
+		}
+
+		// Set up a hover for pins that is tinted the color of the pin.
+		SBorder::Construct(SBorder::FArguments()
+			.BorderImage(this, &SCopyNodeGraphPin::GetPinBorder)
+			.BorderBackgroundColor(this, &SCopyNodeGraphPin::GetPinColor)
+			.OnMouseButtonDown(this, &SCopyNodeGraphPin::OnPinNameMouseDown)
+			[
+				SNew(SLevelOfDetailBranchNode)
+					.UseLowDetailSlot(this, &SCopyNodeGraphPin::UseLowDetailPinNames)
+					.LowDetail()
+					[
+						//@TODO: Try creating a pin-colored line replacement that doesn't measure text / call delegates but still renders
+						ActualPinWidget
+					]
+					.HighDetail()
+					[
+						PinContent.ToSharedRef()
+					]
+			]
+		);
+
+		TAttribute<FText> ToolTipAttribute = TAttribute<FText>::Create(TAttribute<FText>::FGetter::CreateSP(this, &SCopyNodeGraphPin::GetTooltipText));
+		SetToolTipText(ToolTipAttribute);
+	}
+
+	virtual const FSlateBrush* GetPinIcon() const
+	{
+		if (GetDirection() == EGPD_Input)
+		{
+			if (IsConnected())
+			{
+				return FEditorStyle::GetBrush(TEXT("Graph.Pin.CopyNodePinLeft_Connected"));
+			}
+			else
+			{
+				return FEditorStyle::GetBrush(TEXT("Graph.Pin.CopyNodePinLeft_Disconnected"));
+			}
+		}
+		else
+		{
+			if (IsConnected())
+			{
+				return FEditorStyle::GetBrush(TEXT("Graph.Pin.CopyNodePinRight_Connected"));
+			}
+			else
+			{
+				return FEditorStyle::GetBrush(TEXT("Graph.Pin.CopyNodePinRight_Disconnected"));
+			}
+		}
+	}
+
+protected:
+	FSlateBrush* CachedImg_Left_Pin;
+	FSlateBrush* CachedImg_Right_Pin;
+};
+
+//////////////////////////////////////////////////////////////////////////
+// SGraphNodeK2Copy
+
+void SGraphNodeK2Copy::Construct(const FArguments& InArgs, UK2Node* InNode)
+{
+	GraphNode = InNode;
+
+	SetCursor(EMouseCursor::CardinalCross);
+
+	UpdateGraphNode();
+}
+
+void SGraphNodeK2Copy::UpdateGraphNode()
+{
+	InputPins.Empty();
+	OutputPins.Empty();
+
+	// Reset variables that are going to be exposed, in case we are refreshing an already setup node.
+	RightNodeBox.Reset();
+	LeftNodeBox.Reset();
+
+	SetupErrorReporting();
+
+	GetOrAddSlot(ENodeZone::Center)
+		.HAlign(HAlign_Center)
+		.VAlign(VAlign_Center)
+		[
+			CreateNodeContentArea()
+		];
+	CreatePinWidgets();
+}
+
+TSharedRef<SWidget> SGraphNodeK2Copy::CreateNodeContentArea()
+{
+	// NODE CONTENT AREA
+	return SNew(SHorizontalBox)
+		+ SHorizontalBox::Slot()
+		.HAlign(HAlign_Left)
+		.AutoWidth()
+		.Padding(5.0f, 5.0f, 0.0f, 5.0f)
+		[
+			// LEFT
+			SAssignNew(LeftNodeBox, SVerticalBox)
+		]
+	+ SHorizontalBox::Slot()
+		.AutoWidth()
+		.HAlign(HAlign_Right)
+		.Padding(0.0f, 5.0f, 5.0f, 5.0f)
+		[
+			// RIGHT
+			SAssignNew(RightNodeBox, SVerticalBox)
+		];
+}
+
+void SGraphNodeK2Copy::AddPin(const TSharedRef<SGraphPin>& PinToAdd)
+{
+	PinToAdd->SetOwner(SharedThis(this));
+
+	const UEdGraphPin* PinObj = PinToAdd->GetPinObj();
+	const bool bAdvancedParameter = (PinObj != nullptr) && PinObj->bAdvancedView;
+	if (bAdvancedParameter)
+	{
+		PinToAdd->SetVisibility(TAttribute<EVisibility>(PinToAdd, &SGraphPin::IsPinVisibleAsAdvanced));
+	}
+
+	if (PinToAdd->GetDirection() == EEdGraphPinDirection::EGPD_Input)
+	{
+		LeftNodeBox->AddSlot()
+			.AutoHeight()
+			.HAlign(HAlign_Left)
+			.VAlign(VAlign_Center)
+			[
+				PinToAdd
+			];
+		InputPins.Add(PinToAdd);
+	}
+	else // Direction == EEdGraphPinDirection::EGPD_Output
+	{
+		RightNodeBox->AddSlot()
+			.AutoHeight()
+			.HAlign(HAlign_Right)
+			.VAlign(VAlign_Center)
+			[
+				PinToAdd
+			];
+		OutputPins.Add(PinToAdd);
+	}
+}
+
+TSharedPtr<SGraphPin> SGraphNodeK2Copy::CreatePinWidget(UEdGraphPin* Pin) const
+{
+	return SNew(SCopyNodeGraphPin, Pin);
+}
+
+const FSlateBrush* SGraphNodeK2Copy::GetShadowBrush(bool bSelected) const
+{
+	return bSelected ? FEditorStyle::GetBrush(TEXT("Graph.Node.ShadowSelected")) : FEditorStyle::GetNoBrush();
+}
+
+#undef LOCTEXT_NAMESPACE
+
+
diff --git a/Engine/Source/Editor/GraphEditor/Private/KismetNodes/SGraphNodeK2Copy.h b/Engine/Source/Editor/GraphEditor/Private/KismetNodes/SGraphNodeK2Copy.h
new file mode 100644
--- /dev/null
+++ b/Engine/Source/Editor/GraphEditor/Private/KismetNodes/SGraphNodeK2Copy.h
@@ -0,0 +1,20 @@
+// Copyright 1998-2016 Epic Games, Inc. All Rights Reserved.
+
+#pragma once
+
+class GRAPHEDITOR_API SGraphNodeK2Copy : public SGraphNodeK2Base
+{
+public:
+	SLATE_BEGIN_ARGS(SGraphNodeK2Copy) {}
+	SLATE_END_ARGS()
+
+	void Construct(const FArguments& InArgs, UK2Node* InNode);
+
+	//~ Begin SGraphNode Interface
+	virtual void UpdateGraphNode() override;
+	virtual TSharedRef<SWidget> CreateNodeContentArea() override;
+	virtual void AddPin(const TSharedRef<SGraphPin>& PinToAdd) override;
+	virtual const FSlateBrush* GetShadowBrush(bool bSelected) const override;
+	virtual TSharedPtr<SGraphPin> CreatePinWidget(UEdGraphPin* Pin) const override;
+	//~ End SGraphNode Interface
+};
diff --git a/Engine/Source/Editor/GraphEditor/Private/KismetPins/SGraphPinColor.cpp b/Engine/Source/Editor/GraphEditor/Private/KismetPins/SGraphPinColor.cpp
--- a/Engine/Source/Editor/GraphEditor/Private/KismetPins/SGraphPinColor.cpp
+++ b/Engine/Source/Editor/GraphEditor/Private/KismetPins/SGraphPinColor.cpp
@@ -26,7 +26,7 @@ TSharedRef<SWidget>	SGraphPinColor::GetDefaultValueWidget()
 			[
 				SNew( SColorBlock )
 				.Color( this, &SGraphPinColor::GetColor )
-				.IgnoreAlpha(true)
+				.ShowBackgroundForAlpha(true)
 				.OnMouseButtonDown( this, &SGraphPinColor::OnColorBoxClicked )
 			]
 		];
diff --git a/Engine/Source/Editor/GraphEditor/Private/NodeFactory.cpp b/Engine/Source/Editor/GraphEditor/Private/NodeFactory.cpp
--- a/Engine/Source/Editor/GraphEditor/Private/NodeFactory.cpp
+++ b/Engine/Source/Editor/GraphEditor/Private/NodeFactory.cpp
@@ -28,6 +28,7 @@
 #include "KismetNodes/SGraphNodeK2Event.h"
 #include "KismetNodes/SGraphNodeFormatText.h"
 #include "KismetNodes/SGraphNodeK2ArrayFunction.h"
+#include "KismetNodes/SGraphNodeK2Copy.h"
 
 #include "AnimGraphNode_Root.h"
 #include "AnimGraphNode_SequencePlayer.h"
@@ -39,6 +40,7 @@
 #include "AnimStateTransitionNode.h"
 
 #include "AnimationStateMachineSchema.h"
+#include "AnimationGraphSchema.h"
 
 #include "AnimationStateNodes/SGraphNodeAnimState.h"
 #include "AnimationStateNodes/SGraphNodeAnimTransition.h"
@@ -76,6 +78,12 @@
 
 #include "MaterialPins/SGraphPinMaterialInput.h"
 
+#include "ConnectionDrawingPolicy.h"
+#include "AnimGraphConnectionDrawingPolicy.h"
+#include "StateMachineConnectionDrawingPolicy.h"
+#include "SoundCueGraphConnectionDrawingPolicy.h"
+#include "MaterialGraphConnectionDrawingPolicy.h"
+
 #include "EdGraphUtilities.h"
 
 TSharedPtr<SGraphNode> FNodeFactory::CreateNodeWidget(UEdGraphNode* InNode)
@@ -216,6 +224,10 @@ TSharedPtr<SGraphNode> FNodeFactory::CreateNodeWidget(UEdGraphNode* InNode)
 		{
 			return SNew(SGraphNodeKnot, Knot);
 		}
+		else if (UK2Node_Copy* CopyNode = Cast<UK2Node_Copy>(InNode))
+		{
+			return SNew(SGraphNodeK2Copy, CopyNode);
+		}
 		else
 		{
 			return SNew(SGraphNodeK2Default, K2Node);
@@ -424,3 +436,57 @@ TSharedPtr<SGraphPin> FNodeFactory::CreatePinWidget(UEdGraphPin* InPin)
 	// If we didn't pick a custom pin widget, use an uncustomized basic pin
 	return SNew(SGraphPin, InPin);
 }
+
+FConnectionDrawingPolicy* FNodeFactory::CreateConnectionPolicy(const UEdGraphSchema* Schema, int32 InBackLayerID, int32 InFrontLayerID, float ZoomFactor, const FSlateRect& InClippingRect, FSlateWindowElementList& InDrawElements, UEdGraph* InGraphObj)
+{
+    FConnectionDrawingPolicy* ConnectionDrawingPolicy;
+
+    // First give the schema a chance to provide the connection drawing policy
+    ConnectionDrawingPolicy = Schema->CreateConnectionDrawingPolicy(InBackLayerID, InFrontLayerID, ZoomFactor, InClippingRect, InDrawElements, InGraphObj);
+
+    // First give a shot to the registered connection factories
+    if (!ConnectionDrawingPolicy)
+    {
+        for (auto FactoryIt = FEdGraphUtilities::VisualPinConnectionFactories.CreateIterator(); FactoryIt; ++FactoryIt)
+        {
+            TSharedPtr<FGraphPanelPinConnectionFactory> FactoryPtr = *FactoryIt;
+            if (FactoryPtr.IsValid())
+            {
+                ConnectionDrawingPolicy = FactoryPtr->CreateConnectionPolicy(Schema, InBackLayerID, InFrontLayerID, ZoomFactor, InClippingRect, InDrawElements, InGraphObj);
+            }
+        }
+    }
+
+    // If neither the schema nor the factory provides a policy, try the hardcoded ones
+    //@TODO: Fold all of this code into registered factories for the various schemas!
+    if (!ConnectionDrawingPolicy)
+    {
+        if (Schema->IsA(UAnimationGraphSchema::StaticClass()))
+        {
+            ConnectionDrawingPolicy = new FAnimGraphConnectionDrawingPolicy(InBackLayerID, InFrontLayerID, ZoomFactor, InClippingRect, InDrawElements, InGraphObj);
+        }
+        else if (Schema->IsA(UAnimationStateMachineSchema::StaticClass()))
+        {
+            ConnectionDrawingPolicy = new FStateMachineConnectionDrawingPolicy(InBackLayerID, InFrontLayerID, ZoomFactor, InClippingRect, InDrawElements, InGraphObj);
+        }
+        else if (Schema->IsA(UEdGraphSchema_K2::StaticClass()))
+        {
+            ConnectionDrawingPolicy = new FKismetConnectionDrawingPolicy(InBackLayerID, InFrontLayerID, ZoomFactor, InClippingRect, InDrawElements, InGraphObj);
+        }
+        else if (Schema->IsA(USoundCueGraphSchema::StaticClass()))
+        {
+            ConnectionDrawingPolicy = new FSoundCueGraphConnectionDrawingPolicy(InBackLayerID, InFrontLayerID, ZoomFactor, InClippingRect, InDrawElements, InGraphObj);
+        }
+        else if (Schema->IsA(UMaterialGraphSchema::StaticClass()))
+        {
+            ConnectionDrawingPolicy = new FMaterialGraphConnectionDrawingPolicy(InBackLayerID, InFrontLayerID, ZoomFactor, InClippingRect, InDrawElements, InGraphObj);
+        }
+        else
+        {
+            ConnectionDrawingPolicy = new FConnectionDrawingPolicy(InBackLayerID, InFrontLayerID, ZoomFactor, InClippingRect, InDrawElements);
+        }
+    }
+
+    // If we never picked a custom policy, use the uncustomized standard policy
+    return ConnectionDrawingPolicy;
+}
diff --git a/Engine/Source/Editor/GraphEditor/Private/NodeFactory.h b/Engine/Source/Editor/GraphEditor/Private/NodeFactory.h
--- a/Engine/Source/Editor/GraphEditor/Private/NodeFactory.h
+++ b/Engine/Source/Editor/GraphEditor/Private/NodeFactory.h
@@ -11,4 +11,7 @@ public:
 
 	/** Create a widget for the supplied pin */
 	static TSharedPtr<SGraphPin> CreatePinWidget(UEdGraphPin* InPin);
+
+    /** Create a pin connection factory for the supplied schema */
+    static FConnectionDrawingPolicy* CreateConnectionPolicy(const UEdGraphSchema* Schema, int32 InBackLayerID, int32 InFrontLayerID, float ZoomFactor, const FSlateRect& InClippingRect, FSlateWindowElementList& InDrawElements, UEdGraph* InGraphObj);
 };
diff --git a/Engine/Source/Editor/GraphEditor/Private/SGraphActionMenu.cpp b/Engine/Source/Editor/GraphEditor/Private/SGraphActionMenu.cpp
--- a/Engine/Source/Editor/GraphEditor/Private/SGraphActionMenu.cpp
+++ b/Engine/Source/Editor/GraphEditor/Private/SGraphActionMenu.cpp
@@ -290,7 +290,8 @@ void SGraphActionMenu::Construct( const FArguments& InArgs, bool bIsReadOnly/* =
 		.OnContextMenuOpening(InArgs._OnContextMenuOpening)
 		.OnGetChildren(this, &SGraphActionMenu::OnGetChildrenForCategory)
 		.SelectionMode(ESelectionMode::Single)
-		.OnItemScrolledIntoView(this, &SGraphActionMenu::OnItemScrolledIntoView);
+		.OnItemScrolledIntoView(this, &SGraphActionMenu::OnItemScrolledIntoView)
+		.OnSetExpansionRecursive(this, &SGraphActionMenu::OnSetExpansionRecursive);
 
 
 	this->ChildSlot
@@ -1325,6 +1326,19 @@ bool SGraphActionMenu::HandleSelection( TSharedPtr< FGraphActionNode > &InSelect
 	return bResult;
 }
 
+void SGraphActionMenu::OnSetExpansionRecursive(TSharedPtr<FGraphActionNode> InTreeNode, bool bInIsItemExpanded)
+{
+	if (InTreeNode.IsValid() && InTreeNode->Children.Num())
+	{
+		TreeView->SetItemExpansion(InTreeNode, bInIsItemExpanded);
+
+		for (TSharedPtr<FGraphActionNode> Child : InTreeNode->Children)
+		{
+			OnSetExpansionRecursive(Child, bInIsItemExpanded);
+		}
+	}
+}
+
 /////////////////////////////////////////////////////
 
 #undef LOCTEXT_NAMESPACE
diff --git a/Engine/Source/Editor/GraphEditor/Private/SGraphPanel.cpp b/Engine/Source/Editor/GraphEditor/Private/SGraphPanel.cpp
--- a/Engine/Source/Editor/GraphEditor/Private/SGraphPanel.cpp
+++ b/Engine/Source/Editor/GraphEditor/Private/SGraphPanel.cpp
@@ -13,11 +13,6 @@
 #include "InputChord.h"
 
 #include "ConnectionDrawingPolicy.h"
-#include "BlueprintConnectionDrawingPolicy.h"
-#include "AnimGraphConnectionDrawingPolicy.h"
-#include "SoundCueGraphConnectionDrawingPolicy.h"
-#include "MaterialGraphConnectionDrawingPolicy.h"
-#include "StateMachineConnectionDrawingPolicy.h"
 
 #include "AssetSelection.h"
 #include "ComponentAssetBroker.h"
@@ -332,39 +327,7 @@ int32 SGraphPanel::OnPaint( const FPaintArgs& Args, const FGeometry& AllottedGeo
 	// Draw connections between pins 
 	if (Children.Num() > 0 )
 	{
-
-		//@TODO: Pull this into a factory like the pin and node ones
-		FConnectionDrawingPolicy* ConnectionDrawingPolicy;
-		{
-			ConnectionDrawingPolicy = Schema->CreateConnectionDrawingPolicy(WireLayerId, MaxLayerId, ZoomFactor, MyClippingRect, OutDrawElements, GraphObj);
-			if (!ConnectionDrawingPolicy)
-			{
-				if (Schema->IsA(UAnimationGraphSchema::StaticClass()))
-				{
-					ConnectionDrawingPolicy = new FAnimGraphConnectionDrawingPolicy(WireLayerId, MaxLayerId, ZoomFactor, MyClippingRect, OutDrawElements, GraphObj);
-				}
-				else if (Schema->IsA(UAnimationStateMachineSchema::StaticClass()))
-				{
-					ConnectionDrawingPolicy = new FStateMachineConnectionDrawingPolicy(WireLayerId, MaxLayerId, ZoomFactor, MyClippingRect, OutDrawElements, GraphObj);
-				}
-				else if (Schema->IsA(UEdGraphSchema_K2::StaticClass()))
-				{
-					ConnectionDrawingPolicy = new FKismetConnectionDrawingPolicy(WireLayerId, MaxLayerId, ZoomFactor, MyClippingRect, OutDrawElements, GraphObj);
-				}
-				else if (Schema->IsA(USoundCueGraphSchema::StaticClass()))
-				{
-					ConnectionDrawingPolicy = new FSoundCueGraphConnectionDrawingPolicy(WireLayerId, MaxLayerId, ZoomFactor, MyClippingRect, OutDrawElements, GraphObj);
-				}
-				else if (Schema->IsA(UMaterialGraphSchema::StaticClass()))
-				{
-					ConnectionDrawingPolicy = new FMaterialGraphConnectionDrawingPolicy(WireLayerId, MaxLayerId, ZoomFactor, MyClippingRect, OutDrawElements, GraphObj);
-				}
-				else
-				{
-					ConnectionDrawingPolicy = new FConnectionDrawingPolicy(WireLayerId, MaxLayerId, ZoomFactor, MyClippingRect, OutDrawElements);
-				}
-			}
-		}
+        FConnectionDrawingPolicy* ConnectionDrawingPolicy = FNodeFactory::CreateConnectionPolicy(Schema, WireLayerId, MaxLayerId, ZoomFactor, MyClippingRect, OutDrawElements, GraphObj);
 
 		TArray<TSharedPtr<SGraphPin>> OverridePins;
 		for (const FGraphPinHandle& Handle : PreviewConnectorFromPins)
diff --git a/Engine/Source/Editor/GraphEditor/Public/SGraphActionMenu.h b/Engine/Source/Editor/GraphEditor/Public/SGraphActionMenu.h
--- a/Engine/Source/Editor/GraphEditor/Public/SGraphActionMenu.h
+++ b/Engine/Source/Editor/GraphEditor/Public/SGraphActionMenu.h
@@ -304,6 +304,8 @@ protected:
 	/** Callback when rename text is committed */
 	void OnNameTextCommitted(const FText& NewText, ETextCommit::Type InTextCommit, TWeakPtr< FGraphActionNode > InAction);
 	/** Handler for when an item has scrolled into view after having been requested to do so */
-	void OnItemScrolledIntoView( TSharedPtr<FGraphActionNode> InActionNode, const TSharedPtr<ITableRow>& InWidget );	
+	void OnItemScrolledIntoView( TSharedPtr<FGraphActionNode> InActionNode, const TSharedPtr<ITableRow>& InWidget );
+	/** Callback for expanding tree items recursively */
+	void OnSetExpansionRecursive(TSharedPtr<FGraphActionNode> InTreeNode, bool bInIsItemExpanded);
 };
 
diff --git a/Engine/Source/Editor/GraphEditor/Public/SGraphPin.h b/Engine/Source/Editor/GraphEditor/Public/SGraphPin.h
--- a/Engine/Source/Editor/GraphEditor/Public/SGraphPin.h
+++ b/Engine/Source/Editor/GraphEditor/Public/SGraphPin.h
@@ -147,7 +147,6 @@ protected:
 
 	/** Spawns a FDragConnection or similar class for the pin drag event */
 	virtual TSharedRef<FDragDropOperation> SpawnPinDragEvent(const TSharedRef<class SGraphPanel>& InGraphPanel, const TArray< TSharedRef<SGraphPin> >& InStartingPins, bool bShiftOperation);
-private:
 	
 	/** True if pin can be edited */
 	bool IsEditingEnabled() const;
diff --git a/Engine/Source/Editor/Kismet/Private/BlueprintDetailsCustomization.cpp b/Engine/Source/Editor/Kismet/Private/BlueprintDetailsCustomization.cpp
--- a/Engine/Source/Editor/Kismet/Private/BlueprintDetailsCustomization.cpp
+++ b/Engine/Source/Editor/Kismet/Private/BlueprintDetailsCustomization.cpp
@@ -717,14 +717,11 @@ void FBlueprintVarActionDetails::CustomizeDetails( IDetailLayoutBuilder& DetailL
 			}
 			else
 			{
+				UBlueprint* CurrPropertyOwnerBlueprint = IsVariableInheritedByBlueprint() ? GetBlueprintObj() : GetPropertyOwnerBlueprint();
 				UObject* TargetBlueprintDefaultObject = nullptr;
-				if (IsVariableInheritedByBlueprint())
+				if (CurrPropertyOwnerBlueprint && CurrPropertyOwnerBlueprint->GeneratedClass)
 				{
-					TargetBlueprintDefaultObject = GetBlueprintObj()->GeneratedClass->GetDefaultObject();
-				}
-				else if (UBlueprint* PropertyOwnerBP = GetPropertyOwnerBlueprint())
-				{
-					TargetBlueprintDefaultObject = PropertyOwnerBP->GeneratedClass->GetDefaultObject();
+					TargetBlueprintDefaultObject = CurrPropertyOwnerBlueprint->GeneratedClass->GetDefaultObject();
 				}
 				else if (CachedVariableProperty.IsValid())
 				{
diff --git a/Engine/Source/Editor/Kismet/Private/BlueprintEditor.cpp b/Engine/Source/Editor/Kismet/Private/BlueprintEditor.cpp
--- a/Engine/Source/Editor/Kismet/Private/BlueprintEditor.cpp
+++ b/Engine/Source/Editor/Kismet/Private/BlueprintEditor.cpp
@@ -3636,17 +3636,21 @@ void FBlueprintEditor::DoPromoteToVariable( UBlueprint* InBlueprint, UEdGraphPin
 
 	FName VarName;
 	bool bWasSuccessful = false;
+	FEdGraphPinType NewPinType = InTargetPin->PinType;
+	NewPinType.bIsConst = false;
+	NewPinType.bIsReference = false;
+	NewPinType.bIsWeakPointer = false;
 	if (bInToMemberVariable)
 	{
 		VarName = FBlueprintEditorUtils::FindUniqueKismetName(GetBlueprintObj(), TEXT("NewVar"));
-		bWasSuccessful = FBlueprintEditorUtils::AddMemberVariable( GetBlueprintObj(), VarName, InTargetPin->PinType, InTargetPin->GetDefaultAsString() );
+		bWasSuccessful = FBlueprintEditorUtils::AddMemberVariable( GetBlueprintObj(), VarName, NewPinType, InTargetPin->GetDefaultAsString() );
 	}
 	else
 	{
 		ensure(FBlueprintEditorUtils::DoesSupportLocalVariables(GraphObj));
 		VarName = FBlueprintEditorUtils::FindUniqueKismetName(GetBlueprintObj(), TEXT("NewLocalVar"));
 		FunctionGraph = FBlueprintEditorUtils::GetTopLevelGraph(GraphObj);
-		bWasSuccessful = FBlueprintEditorUtils::AddLocalVariable( GetBlueprintObj(), FunctionGraph, VarName, InTargetPin->PinType, InTargetPin->GetDefaultAsString() );
+		bWasSuccessful = FBlueprintEditorUtils::AddLocalVariable( GetBlueprintObj(), FunctionGraph, VarName, NewPinType, InTargetPin->GetDefaultAsString() );
 	}
 
 	if (bWasSuccessful)
diff --git a/Engine/Source/Editor/Kismet/Private/SMyBlueprint.cpp b/Engine/Source/Editor/Kismet/Private/SMyBlueprint.cpp
--- a/Engine/Source/Editor/Kismet/Private/SMyBlueprint.cpp
+++ b/Engine/Source/Editor/Kismet/Private/SMyBlueprint.cpp
@@ -2592,6 +2592,7 @@ void SMyBlueprint::OnDuplicateAction()
 		}
 		FBlueprintEditorUtils::MarkBlueprintAsStructurallyModified(GetBlueprintObj());
 
+		BlueprintEditorPtr.Pin()->OpenDocument(DuplicatedGraph, FDocumentTracker::ForceOpenNewDocument);
 		DuplicateActionName = DuplicatedGraph->GetFName();
 	}
 	else if (FEdGraphSchemaAction_K2Var* VarAction = SelectionAsVar())
diff --git a/Engine/Source/Editor/KismetCompiler/Private/KismetCompiler.cpp b/Engine/Source/Editor/KismetCompiler/Private/KismetCompiler.cpp
--- a/Engine/Source/Editor/KismetCompiler/Private/KismetCompiler.cpp
+++ b/Engine/Source/Editor/KismetCompiler/Private/KismetCompiler.cpp
@@ -3417,6 +3417,13 @@ void FKismetCompilerContext::Compile()
 		CopyTermDefaultsToDefaultObject(NewCDO);
 		SetCanEverTick();
 		SetWantsBeginPlay();
+
+		// Note: The old->new CDO copy is deferred when regenerating, so we skip this step in that case.
+		if (!Blueprint->HasAnyFlags(RF_BeingRegenerated))
+		{
+			// Update the custom property list used in post construction logic to include native class properties for which the Blueprint CDO differs from the native CDO.
+			TargetClass->UpdateCustomPropertyListForPostConstruction();
+		}
 	}
 
 	// Fill out the function bodies, either with function bodies, or simple stubs if this is skeleton generation
diff --git a/Engine/Source/Editor/KismetCompiler/Private/KismetCompilerMisc.cpp b/Engine/Source/Editor/KismetCompiler/Private/KismetCompilerMisc.cpp
--- a/Engine/Source/Editor/KismetCompiler/Private/KismetCompilerMisc.cpp
+++ b/Engine/Source/Editor/KismetCompiler/Private/KismetCompilerMisc.cpp
@@ -1847,6 +1847,10 @@ struct FEventGraphUtils
 		{
 			auto LinkOwnerNode = Cast<const UK2Node>(Link->GetOwningNodeUnchecked());
 			ensure(LinkOwnerNode);
+			if (Link->PinType.bIsReference)
+			{
+				return true;
+			}
 			auto EventsCallingDestination = FEventGraphUtils::FindExecutionNodes(LinkOwnerNode, OwnerNode);
 			if (0 != EventsCallingDestination.Num())
 			{
@@ -1884,17 +1888,16 @@ FBPTerminal* FKismetFunctionContext::CreateLocalTerminalFromPinAutoChooseScope(U
 	check(Net);
 	bool bSharedTerm = IsEventGraph();
 	static FBoolConfigValueHelper UseLocalGraphVariables(TEXT("Kismet"), TEXT("bUseLocalGraphVariables"), GEngineIni);
+	static FBoolConfigValueHelper UseLocalGraphVariablesInCpp(TEXT("BlueprintNativizationSettings"), TEXT("bUseLocalEventGraphVariables"));
+
 
 	auto PinRepresentsDelegate = [](const FEdGraphPinType& PinType) -> bool
 	{
-		return (PinType.PinCategory == UEdGraphSchema_K2::PC_Delegate) 
+		return (PinType.PinCategory == UEdGraphSchema_K2::PC_Delegate)
 			|| (PinType.PinCategory == UEdGraphSchema_K2::PC_MCDelegate);
 	};
-	// Local graph Variables cannot be used with persistent UberGraph frame
-	// Current implementation can change BP behavior, so it cannot be used to generate C++
-	const bool bUseLocalGraphVariables = UseLocalGraphVariables
-		//This is a workaround for UE-24210. UHT doesn't support  delegates with native type form: OtherOwnerClass::DelegateType.
-		|| (bGeneratingCpp && PinRepresentsDelegate(Net->PinType)); 
+	const bool bUseLocalGraphVariables = UseLocalGraphVariables 
+		|| (bGeneratingCpp && (UseLocalGraphVariablesInCpp || PinRepresentsDelegate(Net->PinType)));
 
 	const bool OutputPin = EEdGraphPinDirection::EGPD_Output == Net->Direction;
 	if (bSharedTerm && bUseLocalGraphVariables && OutputPin)
diff --git a/Engine/Source/Editor/UnrealEd/Private/CookOnTheFlyServer.cpp b/Engine/Source/Editor/UnrealEd/Private/CookOnTheFlyServer.cpp
--- a/Engine/Source/Editor/UnrealEd/Private/CookOnTheFlyServer.cpp
+++ b/Engine/Source/Editor/UnrealEd/Private/CookOnTheFlyServer.cpp
@@ -4138,6 +4138,10 @@ void UCookOnTheFlyServer::CookByTheBookFinished()
 		if (IBlueprintNativeCodeGenModule::IsNativeCodeGenModuleLoaded())
 		{
 			IBlueprintNativeCodeGenModule& CodeGenModule = IBlueprintNativeCodeGenModule::Get();
+
+			CodeGenModule.GenerateFullyConvertedClasses(); // While generating fully converted classes the list of necessary stubs is created.
+			CodeGenModule.GenerateStubs();
+
 			// merge the manifest for the blueprint code generator:
 			for (int32 I = 0; I < CookByTheBookOptions->ChildCookers.Num(); ++I )
 			{
diff --git a/Engine/Source/Editor/UnrealEd/Private/EdGraphUtilities.cpp b/Engine/Source/Editor/UnrealEd/Private/EdGraphUtilities.cpp
--- a/Engine/Source/Editor/UnrealEd/Private/EdGraphUtilities.cpp
+++ b/Engine/Source/Editor/UnrealEd/Private/EdGraphUtilities.cpp
@@ -78,6 +78,7 @@ protected:
 
 TArray< TSharedPtr<FGraphPanelNodeFactory> > FEdGraphUtilities::VisualNodeFactories;
 TArray< TSharedPtr<FGraphPanelPinFactory> > FEdGraphUtilities::VisualPinFactories;
+TArray< TSharedPtr<FGraphPanelPinConnectionFactory> > FEdGraphUtilities::VisualPinConnectionFactories;
 
 // Reconcile other pin links:
 //   - Links between nodes within the copied set are fine
@@ -384,6 +385,16 @@ void FEdGraphUtilities::UnregisterVisualPinFactory(TSharedPtr<FGraphPanelPinFact
 	VisualPinFactories.Remove(OldFactory);
 }
 
+void FEdGraphUtilities::RegisterVisualPinConnectionFactory(TSharedPtr<FGraphPanelPinConnectionFactory> NewFactory)
+{
+    VisualPinConnectionFactories.Add(NewFactory);
+}
+
+void FEdGraphUtilities::UnregisterVisualPinConnectionFactory(TSharedPtr<FGraphPanelPinConnectionFactory> OldFactory)
+{
+    VisualPinConnectionFactories.Remove(OldFactory);
+}
+
 //////////////////////////////////////////////////////////////////////////
 // FWeakGraphPinPtr
 
@@ -426,4 +437,4 @@ UEdGraphPin* FWeakGraphPinPtr::Get()
 	}
 
 	return NULL;
-}
\ No newline at end of file
+}
diff --git a/Engine/Source/Editor/UnrealEd/Private/Kismet2/BlueprintEditorUtils.cpp b/Engine/Source/Editor/UnrealEd/Private/Kismet2/BlueprintEditorUtils.cpp
--- a/Engine/Source/Editor/UnrealEd/Private/Kismet2/BlueprintEditorUtils.cpp
+++ b/Engine/Source/Editor/UnrealEd/Private/Kismet2/BlueprintEditorUtils.cpp
@@ -1688,6 +1688,12 @@ UClass* FBlueprintEditorUtils::RegenerateBlueprintClass(UBlueprint* Blueprint, U
 						PatchCDOSubobjectsIntoExport(PreviousCDO, NewCDO);
 						// We purposefully do not call post load here, it happens later on in the normal flow
 					}
+
+					// Update the custom property list used in post construction logic to include native class properties for which the regenerated Blueprint CDO now differs from the native CDO.
+					if (UBlueprintGeneratedClass* BPGClass = Cast<UBlueprintGeneratedClass>(AuthoritativeClass))
+					{
+						BPGClass->UpdateCustomPropertyListForPostConstruction();
+					}
 				}
 
 				Blueprint->PRIVATE_InnermostPreviousCDO = NULL;
@@ -8352,5 +8358,57 @@ void FBlueprintEditorUtils::HandleDisableEditableWhenInherited(UObject* Modified
 	}
 }
 
+void FBlueprintEditorUtils::BuildComponentInstancingData(UActorComponent* ComponentTemplate, FBlueprintCookedComponentInstancingData& OutData)
+{
+	// Recursively gathers properties that differ from class/struct defaults, and fills out the cooked property list structure.
+	TFunction<void(UStruct*, const uint8*, const uint8*)> RecursivePropertyGatherLambda = [&](UStruct* InStruct, const uint8* DataPtr, const uint8* DefaultDataPtr)
+	{
+		for (UProperty* Property = InStruct->PropertyLink; Property; Property = Property->PropertyLinkNext)
+		{
+			if (!Property->IsEditorOnlyProperty())
+			{
+				for (int32 Idx = 0; Idx < Property->ArrayDim; Idx++)
+				{
+					const uint8* PropertyValue = Property->ContainerPtrToValuePtr<uint8>(DataPtr, Idx);
+					const uint8* DefaultPropertyValue = Property->ContainerPtrToValuePtrForDefaults<uint8>(InStruct, DefaultDataPtr, Idx);
+
+					FBlueprintComponentChangedPropertyInfo ChangedPropertyInfo;
+					ChangedPropertyInfo.PropertyName = Property->GetFName();
+					ChangedPropertyInfo.ArrayIndex = Idx;
+					ChangedPropertyInfo.PropertyScope = InStruct;
+
+					if (UStructProperty* StructProperty = Cast<UStructProperty>(Property))
+					{
+						int32 NumChangedProperties = OutData.ChangedPropertyList.Num();
+
+						RecursivePropertyGatherLambda(StructProperty->Struct, PropertyValue, DefaultPropertyValue);
+
+						if (NumChangedProperties < OutData.ChangedPropertyList.Num())
+						{
+							OutData.ChangedPropertyList.Insert(ChangedPropertyInfo, NumChangedProperties);
+						}
+					}
+					else if (!Property->Identical(PropertyValue, DefaultPropertyValue, PPF_None))
+					{
+						OutData.ChangedPropertyList.Add(ChangedPropertyInfo);
+					}
+				}
+			}
+		}
+	};
+
+	if (ComponentTemplate)
+	{
+		UClass* ComponentTemplateClass = ComponentTemplate->GetClass();
+
+		// Gather the set of properties that differ from the template CDO.
+		OutData.ChangedPropertyList.Empty();
+		RecursivePropertyGatherLambda(ComponentTemplateClass, (uint8*)ComponentTemplate, (uint8*)ComponentTemplateClass->GetDefaultObject(false));
+
+		// Flag that cooked data has been built and is now considered to be valid.
+		OutData.bIsValid = true;
+	}
+}
+
 #undef LOCTEXT_NAMESPACE
 
diff --git a/Engine/Source/Editor/UnrealEd/Private/Kismet2/Kismet2.cpp b/Engine/Source/Editor/UnrealEd/Private/Kismet2/Kismet2.cpp
--- a/Engine/Source/Editor/UnrealEd/Private/Kismet2/Kismet2.cpp
+++ b/Engine/Source/Editor/UnrealEd/Private/Kismet2/Kismet2.cpp
@@ -725,16 +725,10 @@ void FKismetEditorUtilities::CompileBlueprint(UBlueprint* BlueprintObj, bool bIs
 	FCompilerResultsLog& Results = (pResults != NULL) ? *pResults : LocalResults;
 
 	// Monitoring UE-20486, the OldClass->ClassGeneratedBy is NULL or otherwise not a UBlueprint.
-	if (OldClass)
+	if (OldClass && (OldClass->ClassGeneratedBy != BlueprintObj))
 	{
-		if (OldClass->ClassGeneratedBy == nullptr)
-		{
-			checkf(OldClass->ClassGeneratedBy == BlueprintObj, TEXT("Generated Class '%s' during compilation has a NULL ClassGeneratedBy for Blueprint '%s'"), *OldClass->GetPathName(), *BlueprintObj->GetPathName());
-		}
-		else
-		{
-			checkf(OldClass->ClassGeneratedBy == BlueprintObj, TEXT("Generated Class '%s' has an invalid ClassGeneratedBy '%s' while the expected is Blueprint '%s'"), *OldClass->GetPathName(), *OldClass->ClassGeneratedBy->GetPathName(), *BlueprintObj->GetPathName());
-		}
+		ensureMsgf(OldClass->ClassGeneratedBy == BlueprintObj, TEXT("Generated Class '%s' has an invalid ClassGeneratedBy '%s' while the expected is Blueprint '%s'"), *OldClass->GetPathName(), *GetPathNameSafe(OldClass->ClassGeneratedBy), *BlueprintObj->GetPathName());
+		OldClass->ClassGeneratedBy = BlueprintObj;
 	}
 	auto ReinstanceHelper = FBlueprintCompileReinstancer::Create(OldClass);
 
diff --git a/Engine/Source/Editor/UnrealEd/Public/CookerSettings.h b/Engine/Source/Editor/UnrealEd/Public/CookerSettings.h
--- a/Engine/Source/Editor/UnrealEd/Public/CookerSettings.h
+++ b/Engine/Source/Editor/UnrealEd/Public/CookerSettings.h
@@ -31,6 +31,10 @@ public:
 	/** Whether or not to compile Blueprints in development mode when cooking. */
 	UPROPERTY(GlobalConfig, EditAnywhere, Category = Cooker, AdvancedDisplay)
 	bool bCompileBlueprintsInDevelopmentMode;
+	
+	/** Whether or not to cook Blueprint Component data for faster instancing at runtime. This assumes that the Component templates do not get modified at runtime. */
+	UPROPERTY(GlobalConfig, EditAnywhere, Category = Cooker, AdvancedDisplay, meta = (DisplayName = "Cook Blueprint Component data for faster instancing at runtime"))
+	bool bCookBlueprintComponentTemplateData;
 
 	/** Quality of 0 means fastest, 4 means best quality */
 	UPROPERTY(GlobalConfig, EditAnywhere, Category = Textures, meta = (DisplayName = "PVRTC Compression Quality (0-4, 0 is fastest)"))
diff --git a/Engine/Source/Editor/UnrealEd/Public/EdGraphUtilities.h b/Engine/Source/Editor/UnrealEd/Public/EdGraphUtilities.h
--- a/Engine/Source/Editor/UnrealEd/Public/EdGraphUtilities.h
+++ b/Engine/Source/Editor/UnrealEd/Public/EdGraphUtilities.h
@@ -24,6 +24,15 @@ public:
 
 //////////////////////////////////////////////////////////////////////////
 
+struct UNREALED_API FGraphPanelPinConnectionFactory : public TSharedFromThis<FGraphPanelPinConnectionFactory>
+{
+public:
+    virtual ~FGraphPanelPinConnectionFactory() {}
+    virtual class FConnectionDrawingPolicy* CreateConnectionPolicy(const class UEdGraphSchema* Schema, int32 InBackLayerID, int32 InFrontLayerID, float ZoomFactor, const class FSlateRect& InClippingRect, class FSlateWindowElementList& InDrawElements, class UEdGraph* InGraphObj) const { return nullptr; }
+};
+
+//////////////////////////////////////////////////////////////////////////
+
 class UNREALED_API FEdGraphUtilities
 {
 public:
@@ -33,6 +42,9 @@ public:
 	static void RegisterVisualPinFactory(TSharedPtr<FGraphPanelPinFactory> NewFactory);
 	static void UnregisterVisualPinFactory(TSharedPtr<FGraphPanelPinFactory> OldFactory);
 
+    static void RegisterVisualPinConnectionFactory(TSharedPtr<FGraphPanelPinConnectionFactory> NewFactory);
+    static void UnregisterVisualPinConnectionFactory(TSharedPtr<FGraphPanelPinConnectionFactory> OldFactory);
+
 	/** After pasting nodes, need to perform some fixup for pins etc. */
 	static void PostProcessPastedNodes(TSet<class UEdGraphNode*>& SpawnedNodes);
 
@@ -138,6 +150,7 @@ public:
 private:
 	static TArray< TSharedPtr<FGraphPanelNodeFactory> > VisualNodeFactories;
 	static TArray< TSharedPtr<FGraphPanelPinFactory> > VisualPinFactories;
+    static TArray< TSharedPtr<FGraphPanelPinConnectionFactory> > VisualPinConnectionFactories;
 	friend class FNodeFactory;
 
 	// Should never create an instance of this class
diff --git a/Engine/Source/Editor/UnrealEd/Public/Kismet2/BlueprintEditorUtils.h b/Engine/Source/Editor/UnrealEd/Public/Kismet2/BlueprintEditorUtils.h
--- a/Engine/Source/Editor/UnrealEd/Public/Kismet2/BlueprintEditorUtils.h
+++ b/Engine/Source/Editor/UnrealEd/Public/Kismet2/BlueprintEditorUtils.h
@@ -1272,7 +1272,16 @@ public:
 	 * @param PinType		The pin get the icon for.
 	 * @param returns a brush that best represents the icon (or Kismet.VariableList.TypeIcon if none is available )
 	 */
-	static const struct FSlateBrush* GetIconFromPin( const FEdGraphPinType& PinType );
+	static const struct FSlateBrush* GetIconFromPin(const FEdGraphPinType& PinType);
+
+	/**
+	 * Generate component instancing data (for cooked builds).
+	 *
+	 * @param ComponentTemplate	The component template to generate instancing data for.
+	 * @param OutData			The generated component instancing data.
+	 * @return					TRUE if component instancing data was built, FALSE otherwise.
+	 */
+	static void BuildComponentInstancingData(UActorComponent* ComponentTemplate, FBlueprintCookedComponentInstancingData& OutData);
 
 protected:
 	// Removes all NULL graph references from the SubGraphs array and recurses thru the non-NULL ones
diff --git a/Engine/Source/Programs/UnrealHeaderTool/Private/CodeGenerator.cpp b/Engine/Source/Programs/UnrealHeaderTool/Private/CodeGenerator.cpp
--- a/Engine/Source/Programs/UnrealHeaderTool/Private/CodeGenerator.cpp
+++ b/Engine/Source/Programs/UnrealHeaderTool/Private/CodeGenerator.cpp
@@ -1067,6 +1067,10 @@ inline FString GetEventStructParamsName(UObject* Outer, const TCHAR* FunctionNam
 	}
 
 	FString Result = FString::Printf(TEXT("%s_event%s_Parms"), *OuterName, FunctionName);
+	if (Result.Len() && FChar::IsDigit(Result[0]))
+	{
+		Result.InsertAt(0, TCHAR('_'));
+	}
 	return Result;
 }
 
@@ -4311,6 +4315,7 @@ void FNativeClassHeaderGenerator::ExportFunctionThunk(FUHTStringBuilder& RPCWrap
 	}
 
 	RPCWrappers += TEXT("\t\tP_FINISH;") LINE_TERMINATOR;
+	RPCWrappers += TEXT("\t\tP_NATIVE_BEGIN;") LINE_TERMINATOR;
 
 	auto ClassRange = ClassDefinitionRange();
 	if (ClassDefinitionRanges.Contains(Function->GetOwnerClass()))
@@ -4378,6 +4383,7 @@ void FNativeClassHeaderGenerator::ExportFunctionThunk(FUHTStringBuilder& RPCWrap
 	{
 		RPCWrappers.Logf(TEXT("this->%s(%s);") LINE_TERMINATOR, *FunctionData.CppImplName, *ParameterList);
 	}
+	RPCWrappers += TEXT("\t\tP_NATIVE_END;") LINE_TERMINATOR;
 }
 
 FString FNativeClassHeaderGenerator::GetFunctionParameterString(UFunction* Function)
diff --git a/Engine/Source/Programs/UnrealHeaderTool/Private/HeaderParser.cpp b/Engine/Source/Programs/UnrealHeaderTool/Private/HeaderParser.cpp
--- a/Engine/Source/Programs/UnrealHeaderTool/Private/HeaderParser.cpp
+++ b/Engine/Source/Programs/UnrealHeaderTool/Private/HeaderParser.cpp
@@ -360,9 +360,24 @@ namespace
 
 				case EFunctionSpecifier::BlueprintPure:
 				{
+					bool bIsPure = true;
+					if (Specifier.Values.Num() == 1)
+					{
+						FString IsPureStr = Specifier.Values[0];
+						bIsPure = IsPureStr.ToBool();
+					}
+
 					// This function can be called, and is also pure.
 					FuncInfo.FunctionFlags |= FUNC_BlueprintCallable;
-					FuncInfo.FunctionFlags |= FUNC_BlueprintPure;
+
+					if (bIsPure)
+					{
+						FuncInfo.FunctionFlags |= FUNC_BlueprintPure;
+					}
+					else
+					{
+						FuncInfo.bForceBlueprintImpure = true;
+					}
 				}
 				break;
 
@@ -428,6 +443,11 @@ namespace
 		{
 			FError::Throwf(TEXT("SealedEvent cannot be used on Blueprint events"));
 		}
+
+		if (FuncInfo.bForceBlueprintImpure && (FuncInfo.FunctionFlags & FUNC_BlueprintPure) != 0)
+		{
+			FError::Throwf(TEXT("BlueprintPure (or BlueprintPure=true) and BlueprintPure=false should not both appear on the same function, they are mutually exclusive"));
+		}
 	}
 
 	void AddEditInlineMetaData(TMap<FName, FString>& MetaData)
@@ -2718,6 +2738,8 @@ FIndexRange*                    ParsedVarIndexRange
 	TArray<FPropertySpecifier> SpecifiersFound;
 
 	TMap<FName, FString> MetaDataFromNewStyle;
+	bool bNativeConst = false;
+	bool bNativeConstTemplateArg = false;
 
 	const bool bIsParamList = (VariableCategory != EVariableCategory::Member) && MatchIdentifier(TEXT("UPARAM"));
 
@@ -2737,6 +2759,7 @@ FIndexRange*                    ParsedVarIndexRange
 		if (MatchIdentifier(TEXT("const")))
 		{
 			Flags |= CPF_ConstParm;
+			bNativeConst = true;
 		}
 	}
 
@@ -3158,6 +3181,7 @@ FIndexRange*                    ParsedVarIndexRange
 	{
 		//@TODO: UCREMOVAL: Should use this to set the new (currently non-existent) CPF_Const flag appropriately!
 		bUnconsumedConstKeyword = true;
+		bNativeConst = true;
 	}
 
 	if (MatchIdentifier(TEXT("mutable")))
@@ -3291,6 +3315,11 @@ FIndexRange*                    ParsedVarIndexRange
 			FError::Throwf(TEXT("Nested containers are not supported.") );
 		}
 
+		if (VarProperty.MetaData.Find(TEXT("NativeConst")))
+		{
+			bNativeConstTemplateArg = true;
+		}
+
 		OriginalVarTypeFlags |= VarProperty.PropertyFlags & (CPF_ContainsInstancedReference | CPF_InstancedReference); // propagate these to the array, we will fix them later
 		VarType.PropertyFlags = OriginalVarTypeFlags;
 		VarProperty.ArrayType = EArrayType::Dynamic;
@@ -3498,7 +3527,45 @@ FIndexRange*                    ParsedVarIndexRange
 			bStripped = true;
 		}
 
-		if (Struct)
+		auto SetDelegateType = [&](UFunction* InFunction, const FString& InIdentifierStripped)
+		{
+			bHandledType = true;
+
+			VarProperty = FPropertyBase(InFunction->HasAnyFunctionFlags(FUNC_MulticastDelegate) ? CPT_MulticastDelegate : CPT_Delegate);
+			VarProperty.DelegateName = *InIdentifierStripped;
+
+			if (!(Disallow & CPF_InstancedReference))
+			{
+				Flags |= CPF_InstancedReference;
+			}
+		};
+
+		if (!Struct && MatchSymbol(TEXT("::")))
+		{
+			FToken DelegateName;
+			if (GetIdentifier(DelegateName))
+			{
+				UClass* LocalOwnerClass = AllClasses.FindClass(*IdentifierStripped);
+				if (LocalOwnerClass)
+				{
+					TSharedRef<FScope> LocScope = FScope::GetTypeScope(LocalOwnerClass);
+					const FString DelegateIdentifierStripped = GetClassNameWithPrefixRemoved(DelegateName.Identifier);
+					if (UFunction* DelegateFunc = Cast<UFunction>(LocScope->FindTypeByName(*(DelegateIdentifierStripped + HEADER_GENERATED_DELEGATE_SIGNATURE_SUFFIX))))
+					{
+						SetDelegateType(DelegateFunc, DelegateIdentifierStripped);
+					}
+				}
+				else
+				{
+					FError::Throwf(TEXT("Cannot find class '%s', to resolve delegate '%s'"), *IdentifierStripped, DelegateName.Identifier);
+				}
+			}
+		}
+
+		if (bHandledType)
+		{
+		}
+		else if (Struct)
 		{
 			if (bStripped)
 			{
@@ -3534,15 +3601,7 @@ FIndexRange*                    ParsedVarIndexRange
 		}
 		else if (UFunction* DelegateFunc = Cast<UFunction>(Scope->FindTypeByName(*(IdentifierStripped + HEADER_GENERATED_DELEGATE_SIGNATURE_SUFFIX))))
 		{
-			bHandledType = true;
-
-			VarProperty = FPropertyBase( DelegateFunc->HasAnyFunctionFlags(FUNC_MulticastDelegate) ? CPT_MulticastDelegate : CPT_Delegate);
-			VarProperty.DelegateName = *IdentifierStripped;
-
-			if (!(Disallow & CPF_InstancedReference))
-			{
-				Flags |= CPF_InstancedReference;
-			}
+			SetDelegateType(DelegateFunc, IdentifierStripped);
 		}
 		else
 		{
@@ -3584,7 +3643,7 @@ FIndexRange*                    ParsedVarIndexRange
 				MatchIdentifier(TEXT("class"));
 
 				// Also consume const
-				MatchIdentifier(TEXT("const"));
+				bNativeConstTemplateArg |= MatchIdentifier(TEXT("const"));
 				
 				// Find the lazy/weak class
 				FToken InnerClass;
@@ -3689,7 +3748,7 @@ FIndexRange*                    ParsedVarIndexRange
 				if (!(Flags & CPF_UObjectWrapper))
 				{
 					// Const after variable type but before pointer symbol
-					MatchIdentifier(TEXT("const"));
+					bNativeConst |= MatchIdentifier(TEXT("const"));
 
 					RequireSymbol(TEXT("*"), TEXT("Expected a pointer type"));
 
@@ -3713,15 +3772,7 @@ FIndexRange*                    ParsedVarIndexRange
 		{
 			if (UFunction* DelegateFunc = (UFunction*)StaticFindObject(UFunction::StaticClass(), ANY_PACKAGE, *(IdentifierStripped + HEADER_GENERATED_DELEGATE_SIGNATURE_SUFFIX)))
 			{
-				bHandledType = true;
-
-				VarProperty = FPropertyBase( DelegateFunc->HasAnyFunctionFlags(FUNC_MulticastDelegate) ? CPT_MulticastDelegate : CPT_Delegate);
-				VarProperty.DelegateName = *IdentifierStripped;
-					
-				if (!(Disallow & CPF_InstancedReference))
-				{
-					Flags |= CPF_InstancedReference;
-				}
+				SetDelegateType(DelegateFunc, IdentifierStripped);
 			}
 
 			if (!bHandledType)
@@ -3737,6 +3788,7 @@ FIndexRange*                    ParsedVarIndexRange
 		if (MatchIdentifier(TEXT("const")))
 		{
 			Flags |= CPF_ConstParm;
+			bNativeConst = true;
 		}
 	}
 
@@ -3776,7 +3828,7 @@ FIndexRange*                    ParsedVarIndexRange
 	if (VariableCategory == EVariableCategory::Member && OwnerStruct->IsA<UClass>() && ((UClass*)OwnerStruct)->HasAnyClassFlags(CLASS_Const))
 	{
 		// Eat a 'not quite truthful' const after the type; autogenerated for member variables of const classes.
-		MatchIdentifier(TEXT("const"));
+		bNativeConst |= MatchIdentifier(TEXT("const"));
 	}
 
 	// Arrays are passed by reference but are only implicitly so; setting it explicitly could cause a problem with replicated functions
@@ -3906,7 +3958,15 @@ FIndexRange*                    ParsedVarIndexRange
 	}
 
 	VarProperty.MetaData = MetaDataFromNewStyle;
-
+	if (bNativeConst)
+	{
+		VarProperty.MetaData.Add(TEXT("NativeConst"), TEXT(""));
+	}
+	if (bNativeConstTemplateArg)
+	{
+		VarProperty.MetaData.Add(TEXT("NativeConstTemplateArg"), TEXT(""));
+	}
+	
 	if (ParsedVarIndexRange)
 	{
 		ParsedVarIndexRange->Count = InputPos - ParsedVarIndexRange->StartIndex;
@@ -5857,8 +5917,8 @@ void FHeaderParser::CompileFunctionDeclaration(FClasses& AllClasses)
 		// or we could just rely on the use marking things BlueprintPure. Either way, checking the C++
 		// const identifier to determine purity is not desirable. We should remove the following logic:
 
-		// If its a const BlueprintCallable function with some sort of output, mark it as BlueprintPure as well
-		if ( bHasAnyOutputs && ((FuncInfo.FunctionFlags & FUNC_BlueprintCallable) != 0) )
+		// If its a const BlueprintCallable function with some sort of output and is not being marked as an BlueprintPure=false function, mark it as BlueprintPure as well
+		if ( bHasAnyOutputs && ((FuncInfo.FunctionFlags & FUNC_BlueprintCallable) != 0) && !FuncInfo.bForceBlueprintImpure)
 		{
 			FuncInfo.FunctionFlags |= FUNC_BlueprintPure;
 		}
diff --git a/Engine/Source/Programs/UnrealHeaderTool/Private/ParserHelper.h b/Engine/Source/Programs/UnrealHeaderTool/Private/ParserHelper.h
--- a/Engine/Source/Programs/UnrealHeaderTool/Private/ParserHelper.h
+++ b/Engine/Source/Programs/UnrealHeaderTool/Private/ParserHelper.h
@@ -903,6 +903,9 @@ struct FFuncInfo
 	int32		MacroLine;
 	/** Position in file where this function was declared. Points to first char of function name. */
 	int32 InputPos;
+	/** TRUE if the function is being forced to be considered as impure by the user */
+	bool bForceBlueprintImpure;
+
 	//@}
 
 	/** Constructor. */
@@ -919,6 +922,7 @@ struct FFuncInfo
 		, bSealedEvent(false)
 		, MacroLine(-1)
 		, InputPos(-1)
+		, bForceBlueprintImpure(false)
 	{}
 
 	FFuncInfo(const FFuncInfo& Other)
@@ -932,6 +936,7 @@ struct FFuncInfo
 		, RPCResponseId(Other.RPCResponseId)
 		, MacroLine(Other.MacroLine)
 		, InputPos(Other.InputPos)
+		, bForceBlueprintImpure(Other.bForceBlueprintImpure)
 	{
 		Function.Clone(Other.Function);
 		if (FunctionReference)
diff --git a/Engine/Source/Runtime/Core/Private/Serialization/Archive.cpp b/Engine/Source/Runtime/Core/Private/Serialization/Archive.cpp
--- a/Engine/Source/Runtime/Core/Private/Serialization/Archive.cpp
+++ b/Engine/Source/Runtime/Core/Private/Serialization/Archive.cpp
@@ -91,6 +91,8 @@ void FArchive::Reset()
 	ArMaxSerializeSize					= 0;
 	ArIsFilterEditorOnly				= false;
 	ArIsSaveGame						= false;
+	ArCustomPropertyList				= nullptr;
+	ArUseCustomPropertyList				= false;
 	CookingTargetPlatform = nullptr;
 	SerializedProperty = nullptr;
 #if WITH_EDITORONLY_DATA
@@ -135,6 +137,8 @@ void FArchive::CopyTrivialFArchiveStatusMembers(const FArchive& ArchiveToCopy)
 	ArMaxSerializeSize                   = ArchiveToCopy.ArMaxSerializeSize;
 	ArIsFilterEditorOnly                 = ArchiveToCopy.ArIsFilterEditorOnly;
 	ArIsSaveGame                         = ArchiveToCopy.ArIsSaveGame;
+	ArCustomPropertyList				 = ArchiveToCopy.ArCustomPropertyList;
+	ArUseCustomPropertyList				 = ArchiveToCopy.ArUseCustomPropertyList;
 	CookingTargetPlatform                = ArchiveToCopy.CookingTargetPlatform;
 	SerializedProperty = ArchiveToCopy.SerializedProperty;
 #if WITH_EDITORONLY_DATA
@@ -798,7 +802,6 @@ VARARG_BODY( void, FArchive::Logf, const TCHAR*, VARARG_NONE )
 	FMemory::SystemFree( Buffer );
 }
 
-
 /*----------------------------------------------------------------------------
 	Transparent compression/ decompression archives.
 ----------------------------------------------------------------------------*/
diff --git a/Engine/Source/Runtime/Core/Public/Serialization/ArchiveBase.h b/Engine/Source/Runtime/Core/Public/Serialization/ArchiveBase.h
--- a/Engine/Source/Runtime/Core/Public/Serialization/ArchiveBase.h
+++ b/Engine/Source/Runtime/Core/Public/Serialization/ArchiveBase.h
@@ -1074,6 +1074,12 @@ public:
 	/** Max size of data that this archive is allowed to serialize. */
 	int64 ArMaxSerializeSize;
 
+	/** Custom property list attribute. If the flag below is set, only these properties will be iterated during serialization. If NULL, then no properties will be iterated. */
+	const struct FCustomPropertyListNode* ArCustomPropertyList;
+
+	/** Set TRUE to use the custom property list attribute for serialization. */
+	bool ArUseCustomPropertyList;
+
 	class FScopeSetDebugSerializationFlags
 	{
 	private:
diff --git a/Engine/Source/Runtime/CoreUObject/Private/UObject/Class.cpp b/Engine/Source/Runtime/CoreUObject/Private/UObject/Class.cpp
--- a/Engine/Source/Runtime/CoreUObject/Private/UObject/Class.cpp
+++ b/Engine/Source/Runtime/CoreUObject/Private/UObject/Class.cpp
@@ -701,7 +701,8 @@ void UStruct::Link(FArchive& Ar, bool bRelinkExistingProperties)
 			RefLinkPtr = &(*RefLinkPtr)->NextRef;
 		}
 
-		bool bOwnedByNativeClass = Property->GetOwnerClass() && Property->GetOwnerClass()->HasAnyClassFlags(CLASS_Native | CLASS_Intrinsic);
+		const UClass* OwnerClass = Property->GetOwnerClass();
+		bool bOwnedByNativeClass = OwnerClass && OwnerClass->HasAnyClassFlags(CLASS_Native | CLASS_Intrinsic);
 
 		if (!Property->HasAnyPropertyFlags(CPF_IsPlainOldData | CPF_NoDestructor) &&
 			!bOwnedByNativeClass) // these would be covered by the native destructor
@@ -711,8 +712,8 @@ void UStruct::Link(FArchive& Ar, bool bRelinkExistingProperties)
 			DestructorLinkPtr = &(*DestructorLinkPtr)->DestructorLinkNext;
 		}
 
-		// Link references to properties that require their values to be copied from CDO.
-		if ((Property->HasAnyPropertyFlags(CPF_Config) && Property->GetOwnerClass() && !Property->GetOwnerClass()->HasAnyClassFlags(CLASS_PerObjectConfig)))
+		// Link references to properties that require their values to be initialized and/or copied from CDO post-construction. Note that this includes all non-native-class-owned properties.
+		if (OwnerClass && (!bOwnedByNativeClass || (Property->HasAnyPropertyFlags(CPF_Config) && !OwnerClass->HasAnyClassFlags(CLASS_PerObjectConfig))))
 		{
 			*PostConstructLinkPtr = Property;
 			PostConstructLinkPtr = &(*PostConstructLinkPtr)->PostConstructLinkNext;
@@ -722,9 +723,10 @@ void UStruct::Link(FArchive& Ar, bool bRelinkExistingProperties)
 		PropertyLinkPtr = &(*PropertyLinkPtr)->PropertyLinkNext;
 	}
 
-	*PropertyLinkPtr = NULL;
-	*DestructorLinkPtr = NULL;
-	*RefLinkPtr = NULL;
+	*PropertyLinkPtr = nullptr;
+	*DestructorLinkPtr = nullptr;
+	*RefLinkPtr = nullptr;
+	*PostConstructLinkPtr = nullptr;
 }
 
 void UStruct::InitializeStruct(void* InDest, int32 ArrayDim/* = 1*/) const
@@ -792,6 +794,24 @@ void UStruct::SerializeBin( FArchive& Ar, void* Data ) const
 			RefLinkProperty->SerializeBinProperty( Ar, Data );
 		}
 	}
+	else if( Ar.ArUseCustomPropertyList )
+	{
+		const FCustomPropertyListNode* CustomPropertyList = Ar.ArCustomPropertyList;
+		for (auto PropertyNode = CustomPropertyList; PropertyNode; PropertyNode = PropertyNode->PropertyListNext)
+		{
+			UProperty* Property = PropertyNode->Property;
+			if( Property )
+			{
+				// Temporarily set to the sub property list, in case we're serializing a UStruct property.
+				Ar.ArCustomPropertyList = PropertyNode->SubPropertyList;
+
+				Property->SerializeBinProperty(Ar, Data, PropertyNode->ArrayIndex);
+
+				// Restore the original property list.
+				Ar.ArCustomPropertyList = CustomPropertyList;
+			}
+		}
+	}
 	else
 	{
 		for (UProperty* Property = PropertyLink; Property != NULL; Property = Property->PropertyLinkNext)
@@ -1386,15 +1406,20 @@ void UStruct::SerializeTaggedProperties(FArchive& Ar, uint8* Data, UStruct* Defa
 		// Save tagged properties.
 
 		// Iterate over properties in the order they were linked and serialize them.
-		for( UProperty* Property = PropertyLink; Property; Property = Property->PropertyLinkNext )
+		const FCustomPropertyListNode* CustomPropertyNode = Ar.ArUseCustomPropertyList ? Ar.ArCustomPropertyList : nullptr;
+		for (UProperty* Property = Ar.ArUseCustomPropertyList ? (CustomPropertyNode ? CustomPropertyNode->Property : nullptr) : PropertyLink;
+			Property;
+			Property = Ar.ArUseCustomPropertyList ? FCustomPropertyListNode::GetNextPropertyAndAdvance(CustomPropertyNode) : Property->PropertyLinkNext)
 		{
 			if( Property->ShouldSerializeValue(Ar) )
 			{
-				for( int32 Idx=0; Idx<Property->ArrayDim; Idx++ )
+				const int32 LoopMin = CustomPropertyNode ? CustomPropertyNode->ArrayIndex : 0;
+				const int32 LoopMax = CustomPropertyNode ? LoopMin + 1 : Property->ArrayDim;
+				for( int32 Idx = LoopMin; Idx < LoopMax; Idx++ )
 				{
 					uint8* DataPtr      = Property->ContainerPtrToValuePtr           <uint8>(Data, Idx);
 					uint8* DefaultValue = Property->ContainerPtrToValuePtrForDefaults<uint8>(DefaultsStruct, Defaults, Idx);
-					if( !Ar.DoDelta() || Ar.IsTransacting() || (!Defaults && !dynamic_cast<const UClass*>(this)) || !Property->Identical( DataPtr, DefaultValue, Ar.GetPortFlags()) )
+					if( CustomPropertyNode || !Ar.DoDelta() || Ar.IsTransacting() || (!Defaults && !dynamic_cast<const UClass*>(this)) || !Property->Identical( DataPtr, DefaultValue, Ar.GetPortFlags()) )
 					{
 						if (bUseAtomicSerialization)
 						{
@@ -1411,8 +1436,22 @@ void UStruct::SerializeTaggedProperties(FArchive& Ar, uint8* Data, UStruct* Defa
 						// need to know how much data this call to SerializeTaggedProperty consumes, so mark where we are
 						int32 DataOffset = Ar.Tell();
 
+						// if using it, save the current custom property list and switch to its sub property list (in case of UStruct serialization)
+						const FCustomPropertyListNode* SavedCustomPropertyList = nullptr;
+						if(Ar.ArUseCustomPropertyList && CustomPropertyNode)
+						{
+							SavedCustomPropertyList = Ar.ArCustomPropertyList;
+							Ar.ArCustomPropertyList = CustomPropertyNode->SubPropertyList;
+						}
+
 						Tag.SerializeTaggedProperty( Ar, Property, DataPtr, DefaultValue );
 
+						// restore the original custom property list after serializing
+						if (SavedCustomPropertyList)
+						{
+							Ar.ArCustomPropertyList = SavedCustomPropertyList;
+						}
+
 						// set the tag's size
 						Tag.Size = Ar.Tell() - DataOffset;
 
@@ -2271,7 +2310,7 @@ void UScriptStruct::SerializeItem(FArchive& Ar, void* Value, void const* Default
 		if (bUseBinarySerialization)
 		{
 			// Struct is already preloaded above.
-			if (!Ar.IsPersistent() && Ar.GetPortFlags() != 0 && !ShouldSerializeAtomically(Ar))
+			if (!Ar.IsPersistent() && Ar.GetPortFlags() != 0 && !ShouldSerializeAtomically(Ar) && !Ar.ArUseCustomPropertyList)
 			{
 				SerializeBinEx(Ar, Value, Defaults, this);
 			}
diff --git a/Engine/Source/Runtime/CoreUObject/Private/UObject/Obj.cpp b/Engine/Source/Runtime/CoreUObject/Private/UObject/Obj.cpp
--- a/Engine/Source/Runtime/CoreUObject/Private/UObject/Obj.cpp
+++ b/Engine/Source/Runtime/CoreUObject/Private/UObject/Obj.cpp
@@ -1000,7 +1000,7 @@ void UObject::SerializeScriptProperties( FArchive& Ar ) const
 #endif
 		ObjClass->SerializeTaggedProperties(Ar, (uint8*)this, HasAnyFlags(RF_ClassDefaultObject) ? ObjClass->GetSuperClass() : ObjClass, (uint8*)DiffObject, bBreakSerializationRecursion ? this : NULL);
 	}
-	else if ( Ar.GetPortFlags() != 0 )
+	else if ( Ar.GetPortFlags() != 0 && !Ar.ArUseCustomPropertyList )
 	{
 		UObject* DiffObject = GetArchetype();
 		ObjClass->SerializeBinEx( Ar, const_cast<UObject *>(this), DiffObject, DiffObject ? DiffObject->GetClass() : NULL );
diff --git a/Engine/Source/Runtime/CoreUObject/Private/UObject/PropertyStr.cpp b/Engine/Source/Runtime/CoreUObject/Private/UObject/PropertyStr.cpp
--- a/Engine/Source/Runtime/CoreUObject/Private/UObject/PropertyStr.cpp
+++ b/Engine/Source/Runtime/CoreUObject/Private/UObject/PropertyStr.cpp
@@ -6,12 +6,58 @@
 	UStrProperty.
 -----------------------------------------------------------------------------*/
 
+// Necessary to fix Compiler Error C2026 and C1091
+FString UStrProperty::ExportCppHardcodedText(const FString& InSource, const FString& Indent)
+{
+	const FString Source = InSource.ReplaceCharWithEscapedChar();
+	FString Result{};
+	const int32 PreferredLineSize = 256;
+	int32 StartPos = 0;
+
+	const int32 LinesPerString = 16;
+	const bool bUseSubStrings = InSource.Len() > (LinesPerString * PreferredLineSize);
+	int32 LineNum = 0;
+	if (bUseSubStrings)
+	{
+		Result += TEXT("*(FString(");
+	}
+
+	do
+	{
+		if (StartPos)
+		{
+			Result += TEXT("\n");
+			Result += Indent;
+		}
+
+		++LineNum;
+		if (bUseSubStrings && 0 == (LineNum % LinesPerString))
+		{
+			Result += TEXT(") + FString(");
+		}
+
+		int32 WantedSize = FMath::Min<int32>(Source.Len() - StartPos, PreferredLineSize);
+		while (((WantedSize + StartPos) < Source.Len()) && (Source[WantedSize + StartPos - 1] == TCHAR('\\')))
+		{
+			WantedSize++;
+		}
+		Result += FString::Printf(TEXT("TEXT(\"%s\")"), *Source.Mid(StartPos, WantedSize));
+		StartPos += WantedSize;
+	} while (StartPos < Source.Len());
+
+	if (bUseSubStrings)
+	{
+		Result += TEXT("))");
+	}
+	return Result;
+}
+
 void UStrProperty::ExportTextItem( FString& ValueStr, const void* PropertyValue, const void* DefaultValue, UObject* Parent, int32 PortFlags, UObject* ExportRootScope ) const
 {
 	FString& StringValue = *(FString*)PropertyValue;
 	if (0 != (PortFlags & PPF_ExportCpp))
 	{
-		ValueStr += FString::Printf(TEXT("FString(TEXT(\"%s\"))"), *(StringValue.ReplaceCharWithEscapedChar()));
+		ValueStr += FString::Printf(TEXT("FString(%s)"), *ExportCppHardcodedText(StringValue, FString()));
 	}
 	else if (!(PortFlags & PPF_Delimited))
 	{
diff --git a/Engine/Source/Runtime/CoreUObject/Private/UObject/ScriptCore.cpp b/Engine/Source/Runtime/CoreUObject/Private/UObject/ScriptCore.cpp
--- a/Engine/Source/Runtime/CoreUObject/Private/UObject/ScriptCore.cpp
+++ b/Engine/Source/Runtime/CoreUObject/Private/UObject/ScriptCore.cpp
@@ -14,6 +14,13 @@ DEFINE_LOG_CATEGORY_STATIC(LogScriptCore, Log, All);
 
 DECLARE_CYCLE_STAT(TEXT("Blueprint Time"),STAT_BlueprintTime,STATGROUP_Game);
 
+#if TOTAL_OVERHEAD_SCRIPT_STATS
+COREUOBJECT_API FBlueprintEventTimer::FScopedVMTimer* FBlueprintEventTimer::ActiveVMTimer = nullptr;
+COREUOBJECT_API FBlueprintEventTimer::FPausableScopeTimer* FBlueprintEventTimer::ActiveTimer = nullptr;
+
+DEFINE_STAT(STAT_ScriptVmTime_Total);
+DEFINE_STAT(STAT_ScriptNativeTime_Total);
+#endif //TOTAL_OVERHEAD_SCRIPT_STATS
 
 /*-----------------------------------------------------------------------------
 	Globals.
@@ -214,7 +221,13 @@ FString UnicodeToCPPIdentifier(const FString& InName, bool bDeprecated, const TC
 		}
 	}
 
-	Ret = FString(Prefix) + Ret + Postfix;
+	FString PrefixStr(Prefix);
+	//fix for error C2059: syntax error : 'bad suffix on number'
+	if (!PrefixStr.Len() && Ret.Len() && FChar::IsDigit(Ret[0]))
+	{
+		Ret.InsertAt(0, TCHAR('_'));
+	}
+	Ret = PrefixStr + Ret + Postfix;
 	return bDeprecated ? Ret + TEXT("_DEPRECATED") : Ret;
 }
 
@@ -509,6 +522,16 @@ IMPLEMENT_VM_FUNCTION(EX_CallMath, execCallMathFunction);
 
 void UObject::CallFunction( FFrame& Stack, RESULT_DECL, UFunction* Function )
 {
+#if PER_FUNCTION_SCRIPT_STATS
+	const bool bShouldTrackFunction = FThreadStats::IsCollectingData();
+	FScopeCycleCounterUObject FunctionScope(bShouldTrackFunction ? Function : nullptr);
+#endif // PER_FUNCTION_SCRIPT_STATS
+
+#if STATS
+	const bool bShouldTrackObject = FThreadStats::IsCollectingData();
+	FScopeCycleCounterUObject ContextScope(bShouldTrackObject ? this : nullptr);
+#endif
+
 	checkSlow(Function);
 
 	if (Function->FunctionFlags & FUNC_Native)
@@ -556,8 +579,8 @@ void UObject::CallFunction( FFrame& Stack, RESULT_DECL, UFunction* Function )
 			}
 
 			// Call regular native function.
-			FScopeCycleCounterUObject ContextScope(Stack.Object);
-			FScopeCycleCounterUObject FunctionScope(Function);
+			FScopeCycleCounterUObject NativeContextScope(Stack.Object);
+			FScopeCycleCounterUObject NativeFunctionScope(Function);
 
 			Function->Invoke(this, Stack, RESULT_PARAM);
 		}
@@ -729,10 +752,22 @@ void UObject::ProcessInternal( FFrame& Stack, RESULT_DECL )
 		return;
 	}
 
-	int32 FunctionCallspace = GetFunctionCallspace( (UFunction*)Stack.Node, Stack.Locals, NULL );
+	UFunction* Function = (UFunction*)Stack.Node;
+
+#if PER_FUNCTION_SCRIPT_STATS
+	const bool bShouldTrackFunction = FThreadStats::IsCollectingData();
+	FScopeCycleCounterUObject FunctionScope(bShouldTrackFunction ? Function : nullptr);
+#endif // PER_FUNCTION_SCRIPT_STATS
+
+#if STATS
+	const bool bShouldTrackObject = FThreadStats::IsCollectingData();
+	FScopeCycleCounterUObject ContextScope(bShouldTrackObject ? this : nullptr);
+#endif
+
+	int32 FunctionCallspace = GetFunctionCallspace(Function, Stack.Locals, NULL);
 	if (FunctionCallspace & FunctionCallspace::Remote)
 	{
-		CallRemoteFunction((UFunction*)Stack.Node, Stack.Locals, Stack.OutParms, NULL);
+		CallRemoteFunction(Function, Stack.Locals, Stack.OutParms, NULL);
 	}
 
 	if (FunctionCallspace & FunctionCallspace::Local)
@@ -744,7 +779,7 @@ void UObject::ProcessInternal( FFrame& Stack, RESULT_DECL )
 		if(FBlueprintExceptionTracker::Get().bRanaway)
 		{
 			// If we have a return property, return a zeroed value in it, to try and save execution as much as possible
-			UProperty* ReturnProp = ((UFunction*)Stack.Node)->GetReturnProperty();
+			UProperty* ReturnProp = (Function)->GetReturnProperty();
 			ClearReturnValue(ReturnProp, RESULT_PARAM);
 			return;
 		}
@@ -754,7 +789,7 @@ void UObject::ProcessInternal( FFrame& Stack, RESULT_DECL )
 			UE_LOG(LogScriptCore, Log, TEXT("%s"), *Stack.GetStackTrace());
 
 			// If we have a return property, return a zeroed value in it, to try and save execution as much as possible
-			UProperty* ReturnProp = ((UFunction*)Stack.Node)->GetReturnProperty();
+			UProperty* ReturnProp = (Function)->GetReturnProperty();
 			ClearReturnValue(ReturnProp, RESULT_PARAM);
 
 			// Notify anyone who cares that we've had a fatal error, so we can shut down PIE, etc
@@ -769,9 +804,6 @@ void UObject::ProcessInternal( FFrame& Stack, RESULT_DECL )
 			return;
 		}
 #endif
-		FScopeCycleCounterUObject ContextScope(Stack.Object);
-		FScopeCycleCounterUObject FunctionScope((UFunction*)Stack.Node);
-
 		// Execute the bytecode
 		while (*Stack.Code != EX_Return)
 		{
@@ -782,7 +814,7 @@ void UObject::ProcessInternal( FFrame& Stack, RESULT_DECL )
 				UE_LOG(LogScriptCore, Log, TEXT("%s"), *Stack.GetStackTrace());
 
 				// If we have a return property, return a zeroed value in it, to try and save execution as much as possible
-				UProperty* ReturnProp = ((UFunction*)Stack.Node)->GetReturnProperty();
+				UProperty* ReturnProp = (Function)->GetReturnProperty();
 				ClearReturnValue(ReturnProp, RESULT_PARAM);
 
 				// Notify anyone who cares that we've had a fatal error, so we can shut down PIE, etc
@@ -819,7 +851,7 @@ void UObject::ProcessInternal( FFrame& Stack, RESULT_DECL )
 	}
 	else
 	{
-		UProperty* ReturnProp = ((UFunction*)Stack.Node)->GetReturnProperty();
+		UProperty* ReturnProp = (Function)->GetReturnProperty();
 		if (ReturnProp != NULL)
 		{
 			// destroy old value if necessary
@@ -1019,6 +1051,20 @@ void UObject::ProcessEvent( UFunction* Function, void* Parms )
 	}
 	checkSlow((Function->ParmsSize == 0) || (Parms != NULL));
 
+#if TOTAL_OVERHEAD_SCRIPT_STATS
+	FBlueprintEventTimer::FScopedVMTimer VMTime;
+#endif // TOTAL_OVERHEAD_SCRIPT_STATS
+
+#if PER_FUNCTION_SCRIPT_STATS
+	const bool bShouldTrackFunction = FThreadStats::IsCollectingData();
+	FScopeCycleCounterUObject FunctionScope(bShouldTrackFunction ? Function : nullptr);
+#endif // PER_FUNCTION_SCRIPT_STATS
+
+#if STATS
+	const bool bShouldTrackObject = FThreadStats::IsCollectingData();
+	FScopeCycleCounterUObject ContextScope(bShouldTrackObject ? this : nullptr);
+#endif
+
 #if !(UE_BUILD_SHIPPING || UE_BUILD_TEST)
 	if (GetClass()->HasInstrumentation())
 	{
@@ -1128,8 +1174,6 @@ void UObject::ProcessEvent( UFunction* Function, void* Parms )
 		uint8* ReturnValueAdress = bHasReturnParam ? ((uint8*)Parms + Function->ReturnValueOffset) : nullptr;
 		if (Function->FunctionFlags & FUNC_Native)
 		{
-			FScopeCycleCounterUObject ContextScope(this);
-			FScopeCycleCounterUObject FunctionScope(Function);
 			Function->Invoke(this, NewStack, ReturnValueAdress);
 		}
 		else
diff --git a/Engine/Source/Runtime/CoreUObject/Private/UObject/UObjectGlobals.cpp b/Engine/Source/Runtime/CoreUObject/Private/UObject/UObjectGlobals.cpp
--- a/Engine/Source/Runtime/CoreUObject/Private/UObject/UObjectGlobals.cpp
+++ b/Engine/Source/Runtime/CoreUObject/Private/UObject/UObjectGlobals.cpp
@@ -2681,15 +2681,19 @@ void FObjectInitializer::InitProperties(UObject* Obj, UClass* DefaultsClass, UOb
 	check(DefaultsClass && Obj);
 
 	UClass* Class = Obj->GetClass();
+
 	// bool to indicate that we need to initialize any non-native properties (native ones were done when the native constructor was called by the code that created and passed in a FObjectInitializer object)
 	bool bNeedInitialize = !Class->HasAnyClassFlags(CLASS_Native | CLASS_Intrinsic);
 
+	// bool to indicate that we can use the faster PostConstructLink chain for initialization.
+	bool bCanUsePostConstructLink = !bCopyTransientsFromClassDefaults && DefaultsClass == Class;
+
 	if (Obj->HasAnyFlags(RF_NeedLoad))
 	{
 		bCopyTransientsFromClassDefaults = false;
 	}
 
-	if (!bNeedInitialize && !bCopyTransientsFromClassDefaults && DefaultsClass == Class)
+	if (!bNeedInitialize && bCanUsePostConstructLink)
 	{
 		// This is just a fast path for the below in the common case that we are not doing a duplicate or initializing a CDO and this is all native.
 		// We only do it if the DefaultData object is NOT a CDO of the object that's being initialized. CDO data is already initialized in the
@@ -2718,8 +2722,9 @@ void FObjectInitializer::InitProperties(UObject* Obj, UClass* DefaultsClass, UOb
 	else
 	{
 		QUICK_SCOPE_CYCLE_COUNTER(STAT_InitProperties_Blueprint);
+
 		UObject* ClassDefaults = bCopyTransientsFromClassDefaults ? DefaultsClass->GetDefaultObject() : NULL;		
-		for (UProperty* P = Class->PropertyLink; P; P = P->PropertyLinkNext)
+		for (UProperty* P = bCanUsePostConstructLink ? Class->PostConstructLink : Class->PropertyLink; P; P = bCanUsePostConstructLink ? P->PostConstructLinkNext : P->PropertyLinkNext)
 		{
 			if (bNeedInitialize)
 			{		
@@ -2741,6 +2746,38 @@ void FObjectInitializer::InitProperties(UObject* Obj, UClass* DefaultsClass, UOb
 				}
 			}
 		}
+
+		// This step is only necessary if we're not iterating the full property chain.
+		if (bCanUsePostConstructLink)
+		{
+			// Initialize remaining property values from defaults using an explicit custom post-construction property list returned by the class object.
+			if (const FCustomPropertyListNode* CustomPropertyList = Class->GetCustomPropertyListForPostConstruction())
+			{
+				InitPropertiesFromCustomList(CustomPropertyList, Class, (uint8*)Obj, (uint8*)DefaultData);
+			}
+		}
+	}
+}
+
+void FObjectInitializer::InitPropertiesFromCustomList(const FCustomPropertyListNode* InPropertyList, UStruct* InStruct, uint8* DataPtr, const uint8* DefaultDataPtr)
+{
+	for (const FCustomPropertyListNode* CustomPropertyListNode = InPropertyList; CustomPropertyListNode; CustomPropertyListNode = CustomPropertyListNode->PropertyListNext)
+	{
+		uint8* PropertyValue = CustomPropertyListNode->Property->ContainerPtrToValuePtr<uint8>(DataPtr, CustomPropertyListNode->ArrayIndex);
+		const uint8* DefaultPropertyValue = CustomPropertyListNode->Property->ContainerPtrToValuePtr<uint8>(DefaultDataPtr, CustomPropertyListNode->ArrayIndex);
+
+		if (const UStructProperty* StructProperty = Cast<UStructProperty>(CustomPropertyListNode->Property))
+		{
+			// This should never be NULL; we should not be recording the StructProperty without at least one sub property, but we'll verify just to be sure.
+			if (ensure(CustomPropertyListNode->SubPropertyList != nullptr))
+			{
+				InitPropertiesFromCustomList(CustomPropertyListNode->SubPropertyList, StructProperty->Struct, PropertyValue, DefaultPropertyValue);
+			}
+		}
+		else
+		{
+			CustomPropertyListNode->Property->CopySingleValue(PropertyValue, DefaultPropertyValue);
+		}
 	}
 }
 
diff --git a/Engine/Source/Runtime/CoreUObject/Private/UObject/UTextProperty.cpp b/Engine/Source/Runtime/CoreUObject/Private/UObject/UTextProperty.cpp
--- a/Engine/Source/Runtime/CoreUObject/Private/UObject/UTextProperty.cpp
+++ b/Engine/Source/Runtime/CoreUObject/Private/UObject/UTextProperty.cpp
@@ -51,7 +51,7 @@ void UTextProperty::ExportTextItem( FString& ValueStr, const void* PropertyValue
 	const FString& StringValue = FTextInspector::GetDisplayString(TextValue);
 	if (0 != (PortFlags & PPF_ExportCpp))
 	{
-		ValueStr += FString::Printf(TEXT("FText::FromString(TEXT(\"%s\"))"), *(StringValue.ReplaceCharWithEscapedChar()));
+		ValueStr += FString::Printf(TEXT("FText::FromString(%s)"), *UStrProperty::ExportCppHardcodedText(StringValue, FString()));
 	}
 	else if( !(PortFlags & PPF_Delimited) )
 	{
diff --git a/Engine/Source/Runtime/CoreUObject/Public/UObject/AssetPtr.h b/Engine/Source/Runtime/CoreUObject/Public/UObject/AssetPtr.h
--- a/Engine/Source/Runtime/CoreUObject/Public/UObject/AssetPtr.h
+++ b/Engine/Source/Runtime/CoreUObject/Public/UObject/AssetPtr.h
@@ -482,6 +482,24 @@ public:
 		return GetTypeHash(static_cast<const TPersistentObjectPtr<FStringAssetReference>&>(Other.AssetPtr));
 	}
 
+	/**  
+	 * Synchronously load (if necessary) and return the asset object represented by this asset ptr
+	 */
+	UClass* LoadSynchronous()
+	{
+		UObject* Asset = AssetPtr.Get();
+		if (Asset == nullptr && IsPending())
+		{
+			Asset = AssetPtr.LoadSynchronous();
+		}
+		UClass* Class = dynamic_cast<UClass*>(Asset);
+		if (!Class || !Class->IsChildOf(TClass::StaticClass()))
+		{
+			return nullptr;
+		}
+		return Class;
+	}
+
 	friend FArchive& operator<<(FArchive& Ar, TAssetSubclassOf<TClass>& Other)
 	{
 		Ar << static_cast<FAssetPtr&>(Other.AssetPtr);
diff --git a/Engine/Source/Runtime/CoreUObject/Public/UObject/Class.h b/Engine/Source/Runtime/CoreUObject/Public/UObject/Class.h
--- a/Engine/Source/Runtime/CoreUObject/Public/UObject/Class.h
+++ b/Engine/Source/Runtime/CoreUObject/Public/UObject/Class.h
@@ -2187,6 +2187,12 @@ public:
 		return ClassDefaultObject;
 	}
 
+	/**
+	 * Override to return a linked list of properties with default values that differ from the parent default object. If non-NULL, only these properties will 
+	 * be copied post-construction. Otherwise, all properties will be copied to the new instance, even if the default value matches the inherited default value.
+	 */
+	virtual const FCustomPropertyListNode* GetCustomPropertyListForPostConstruction() const { return nullptr; }
+
 	/**
 	* Get the name of the CDO for the this class
 	* @return The name of the CDO
@@ -2364,7 +2370,7 @@ public:
 	 * @param Object the object to serialize as default
 	 * @param Ar the archive to serialize from
 	 */
-	void SerializeDefaultObject(UObject* Object, FArchive& Ar);
+	virtual void SerializeDefaultObject(UObject* Object, FArchive& Ar);
 
 	/** 
 	 * Purges out the properties of this class in preparation for it to be regenerated
diff --git a/Engine/Source/Runtime/CoreUObject/Public/UObject/Script.h b/Engine/Source/Runtime/CoreUObject/Public/UObject/Script.h
--- a/Engine/Source/Runtime/CoreUObject/Public/UObject/Script.h
+++ b/Engine/Source/Runtime/CoreUObject/Public/UObject/Script.h
@@ -9,6 +9,19 @@
 class UStruct;
 struct FFrame;
 
+// It's best to set only one of these, but strictly speaking you could set both.
+// The results will be confusing. Native time would be included only in a coarse 
+// 'native time' timer, and all overhead would be broken up per script function
+#define TOTAL_OVERHEAD_SCRIPT_STATS (STATS && 0)
+#define PER_FUNCTION_SCRIPT_STATS (STATS && 1)
+
+DECLARE_STATS_GROUP(TEXT("Scripting"), STATGROUP_Script, STATCAT_Advanced);
+
+#if TOTAL_OVERHEAD_SCRIPT_STATS
+DECLARE_FLOAT_COUNTER_STAT_EXTERN(TEXT("Blueprint - (All) VM Time (ms)"),     STAT_ScriptVmTime_Total,     STATGROUP_Script, COREUOBJECT_API);
+DECLARE_FLOAT_COUNTER_STAT_EXTERN(TEXT("Blueprint - (All) Native Time (ms)"), STAT_ScriptNativeTime_Total, STATGROUP_Script, COREUOBJECT_API);
+#endif // TOTAL_OVERHEAD_SCRIPT_STATS
+
 /*-----------------------------------------------------------------------------
 	Constants & types.
 -----------------------------------------------------------------------------*/
@@ -388,6 +401,101 @@ private:
 	bool bOldGAllowScriptExecutionInEditor;
 };
 
+#if TOTAL_OVERHEAD_SCRIPT_STATS
+	// Low overhead timer used to instrument the VM (ProcessEvent and ProcessInternal):
+	struct FBlueprintEventTimer
+	{
+		struct FPausableScopeTimer;
+		COREUOBJECT_API static FPausableScopeTimer* ActiveTimer;
+
+		struct FPausableScopeTimer
+		{
+			FPausableScopeTimer() 
+			{ 
+				double CurrentTime = FPlatformTime::Seconds();
+				if (ActiveTimer)
+				{
+					ActiveTimer->Pause(CurrentTime);
+				}
+
+				PreviouslyActiveTimer = ActiveTimer;
+				StartTime = CurrentTime;
+				TotalTime = 0.0;
+
+				ActiveTimer = this;
+			}
+
+			~FPausableScopeTimer()
+			{
+				if (PreviouslyActiveTimer)
+				{
+					PreviouslyActiveTimer->Resume();
+				}
+				ActiveTimer = PreviouslyActiveTimer;
+			}
+
+			void Pause(double CurrentTime) { TotalTime += CurrentTime - StartTime; }
+			void Resume() { StartTime = FPlatformTime::Seconds();  }
+			double Stop() { return TotalTime + (FPlatformTime::Seconds() - StartTime); }
+
+		private:
+			FPausableScopeTimer* PreviouslyActiveTimer;
+			double TotalTime;
+			double StartTime;
+		};
+
+		// We need to keep track of the current VM timer because we only want to
+		// track time while 'in' the VM. We use this to detect whether we're running
+		// script or just doing RPC:
+		struct FScopedVMTimer;
+		COREUOBJECT_API static FScopedVMTimer* ActiveVMTimer;
+
+		struct FScopedVMTimer
+		{
+			FScopedVMTimer()
+				: Timer()
+				, VMParent(ActiveVMTimer)
+			{
+				ActiveVMTimer = this;
+			}
+
+			~FScopedVMTimer()
+			{
+				INC_FLOAT_STAT_BY(STAT_ScriptVmTime_Total, Timer.Stop() * 1000.0);
+				ActiveVMTimer = VMParent;
+			}
+
+			FPausableScopeTimer Timer;
+			FScopedVMTimer* VMParent;
+		};
+
+		struct FScopedNativeTimer
+		{
+			FScopedNativeTimer()
+				: Timer()
+			{
+			}
+
+			~FScopedNativeTimer()
+			{
+				// only track native time when in a VM scope, RPC time
+				// can be tracked by the online system or whatever is making RPCs:
+				if (ActiveVMTimer)
+				{
+					INC_FLOAT_STAT_BY(STAT_ScriptNativeTime_Total, Timer.Stop()* 1000.0);
+				}
+			}
+
+			FPausableScopeTimer Timer;
+		};
+	};
+
+#define SCOPED_SCRIPT_NATIVE_TIMER(VarName) \
+	FBlueprintEventTimer::FScopedNativeTimer VarName
+
+#else  // TOTAL_OVERHEAD_SCRIPT_STATS
+	#define SCOPED_SCRIPT_NATIVE_TIMER(VarName) 
+#endif // TOTAL_OVERHEAD_SCRIPT_STATS
 /** @return True if the char can be used in an identifier in c++ */
 COREUOBJECT_API bool IsValidCPPIdentifierChar(TCHAR Char);
 
diff --git a/Engine/Source/Runtime/CoreUObject/Public/UObject/ScriptMacros.h b/Engine/Source/Runtime/CoreUObject/Public/UObject/ScriptMacros.h
--- a/Engine/Source/Runtime/CoreUObject/Public/UObject/ScriptMacros.h
+++ b/Engine/Source/Runtime/CoreUObject/Public/UObject/ScriptMacros.h
@@ -72,4 +72,7 @@ enum {MAX_VARIABLE_SIZE = 0x0FFF };
 #define P_GET_ARRAY(ElementType,ParamName)			ElementType ParamName[(MAX_VARIABLE_SIZE/sizeof(ElementType))+1];		Stack.StepCompiledIn<UProperty>(ParamName);
 #define P_GET_ARRAY_REF(ElementType,ParamName)		ElementType ParamName##Temp[(MAX_VARIABLE_SIZE/sizeof(ElementType))+1]; ElementType* ParamName = Stack.StepCompiledInRef<UProperty, ElementType*>(ParamName##Temp);
 
-#define P_FINISH									Stack.Code += !!Stack.Code; // increment the code ptr unless it is null
+#define P_FINISH									Stack.Code += !!Stack.Code; /* increment the code ptr unless it is null */
+
+#define P_NATIVE_BEGIN { SCOPED_SCRIPT_NATIVE_TIMER(ScopedNativeCallTimer);
+#define P_NATIVE_END   }
\ No newline at end of file
diff --git a/Engine/Source/Runtime/CoreUObject/Public/UObject/UObjectGlobals.h b/Engine/Source/Runtime/CoreUObject/Public/UObject/UObjectGlobals.h
--- a/Engine/Source/Runtime/CoreUObject/Public/UObject/UObjectGlobals.h
+++ b/Engine/Source/Runtime/CoreUObject/Public/UObject/UObjectGlobals.h
@@ -911,6 +911,16 @@ private:
 	 */
 	static void InitProperties(UObject* Obj, UClass* DefaultsClass, UObject* DefaultData, bool bCopyTransientsFromClassDefaults);
 
+	/**
+	 * Helper method to assist with initializing object properties from an explicit list.
+	 * 
+	 * @param	InPropertyList		only these properties will be copied from defaults
+	 * @param	InStruct			the current scope for which the given property list applies
+	 * @param	DataPtr				destination address (where to start copying values to)
+	 * @param	DefaultDataPtr		source address (where to start copying the defaults data from)
+	 */
+	static void InitPropertiesFromCustomList(const FCustomPropertyListNode* InPropertyList, UStruct* InStruct, uint8* DataPtr, const uint8* DefaultDataPtr);
+
 	bool IsInstancingAllowed() const;
 
 	/**
diff --git a/Engine/Source/Runtime/CoreUObject/Public/UObject/UnrealType.h b/Engine/Source/Runtime/CoreUObject/Public/UObject/UnrealType.h
--- a/Engine/Source/Runtime/CoreUObject/Public/UObject/UnrealType.h
+++ b/Engine/Source/Runtime/CoreUObject/Public/UObject/UnrealType.h
@@ -216,13 +216,16 @@ public:
 	 *
 	 * @param	Ar				the archive to use for serialization
 	 * @param	Data			pointer to the location of the beginning of the struct's property data
+	 * @param	ArrayIdx		if not -1 (default), only this array slot will be serialized
 	 */
-	void SerializeBinProperty( FArchive& Ar, void* Data )
+	void SerializeBinProperty( FArchive& Ar, void* Data, int32 ArrayIdx = -1 )
 	{
 		if( ShouldSerializeValue(Ar) )
 		{
 			FSerializedPropertyScope SerializedProperty(Ar, this);
-			for (int32 Idx = 0; Idx < ArrayDim; Idx++)
+			const int32 LoopMin = ArrayIdx < 0 ? 0 : ArrayIdx;
+			const int32 LoopMax = ArrayIdx < 0 ? ArrayDim : ArrayIdx + 1;
+			for (int32 Idx = LoopMin; Idx < LoopMax; Idx++)
 			{
 				// Keep setting the property in case something inside of SerializeItem changes it
 				Ar.SetSerializedProperty(this);
@@ -2366,6 +2369,9 @@ public:
 		return GetTypeHash(*(const FString*)Src);
 	}
 	// End of UProperty interface
+
+	// Necessary to fix Compiler Error C2026
+	static FString ExportCppHardcodedText(const FString& InSource, const FString& Indent);
 };
 
 /*-----------------------------------------------------------------------------
@@ -3390,6 +3396,43 @@ protected:
 };
 
 
+/** Describes a single node in a custom property list. */
+struct COREUOBJECT_API FCustomPropertyListNode
+{
+	/** The property that's being referenced at this node. */
+	UProperty* Property;
+
+	/** Used to identify which array index is specifically being referenced if this is an array property. Defaults to 0. */
+	int32 ArrayIndex;
+
+	/** If this node represents a struct property, this may contain a "sub" property list for the struct itself. */
+	struct FCustomPropertyListNode* SubPropertyList;
+
+	/** Points to the next node in the list. */
+	struct FCustomPropertyListNode* PropertyListNext;
+
+	/** Default constructor. */
+	FCustomPropertyListNode(UProperty* InProperty = nullptr, int32 InArrayIndex = 0)
+		:Property(InProperty)
+		, ArrayIndex(InArrayIndex)
+		, SubPropertyList(nullptr)
+		, PropertyListNext(nullptr)
+	{
+	}
+
+	/** Convenience method to return the next property in the list and advance the given ptr. */
+	FORCEINLINE static UProperty* GetNextPropertyAndAdvance(const FCustomPropertyListNode*& Node)
+	{
+		if (Node)
+		{
+			Node = Node->PropertyListNext;
+		}
+
+		return Node ? Node->Property : nullptr;
+	}
+};
+
+
 /**
  * This class represents the chain of member properties leading to an internal struct property.  It is used
  * for tracking which member property corresponds to the UScriptStruct that owns a particular property.
diff --git a/Engine/Source/Runtime/Engine/Classes/Components/SkeletalMeshComponent.h b/Engine/Source/Runtime/Engine/Classes/Components/SkeletalMeshComponent.h
--- a/Engine/Source/Runtime/Engine/Classes/Components/SkeletalMeshComponent.h
+++ b/Engine/Source/Runtime/Engine/Classes/Components/SkeletalMeshComponent.h
@@ -352,7 +352,7 @@ public:
 #endif
 
 	DEPRECATED(4.11, "This property is deprecated. Please use AnimClass instead")
-	UPROPERTY(EditAnywhere, BlueprintReadOnly, Category = Animation)
+	UPROPERTY(BlueprintReadOnly, Category = Animation, meta = (DeprecationMessage = "This property is deprecated. Please use AnimClass instead"))
 	class UAnimBlueprintGeneratedClass* AnimBlueprintGeneratedClass;
 
 	/* The AnimBlueprint class to use. Use 'SetAnimInstanceClass' to change at runtime. */
diff --git a/Engine/Source/Runtime/Engine/Classes/Engine/Blueprint.h b/Engine/Source/Runtime/Engine/Classes/Engine/Blueprint.h
--- a/Engine/Source/Runtime/Engine/Classes/Engine/Blueprint.h
+++ b/Engine/Source/Runtime/Engine/Classes/Engine/Blueprint.h
@@ -333,7 +333,7 @@ class ENGINE_API UBlueprint : public UBlueprintCore
 	UPROPERTY(EditAnywhere, Category=BlueprintOptions)
 	FString BlueprintCategory;
 
-	/** Additional HideCategories. The are added to HideCategories from parent. */
+	/** Additional HideCategories. These are added to HideCategories from parent. */
 	UPROPERTY(EditAnywhere, Category=BlueprintOptions)
 	TArray<FString> HideCategories;
 
@@ -591,6 +591,9 @@ public:
 	virtual void PostLoadSubobjects( FObjectInstancingGraph* OuterInstanceGraph ) override;
 	virtual bool Modify(bool bAlwaysMarkDirty = true) override;
 	virtual void GetAssetRegistryTags(TArray<FAssetRegistryTag>& OutTags) const override;
+	virtual void BeginCacheForCookedPlatformData(const ITargetPlatform *TargetPlatform) override;
+	virtual bool IsCachedCookedPlatformDataLoaded(const ITargetPlatform* TargetPlatform) override;
+	virtual void ClearAllCachedCookedPlatformData() override;
 	//~ End UObject Interface
 
 	/** Consigns the GeneratedClass and the SkeletonGeneratedClass to oblivion, and nulls their references */
diff --git a/Engine/Source/Runtime/Engine/Classes/Engine/BlueprintGeneratedClass.h b/Engine/Source/Runtime/Engine/Classes/Engine/BlueprintGeneratedClass.h
--- a/Engine/Source/Runtime/Engine/Classes/Engine/BlueprintGeneratedClass.h
+++ b/Engine/Source/Runtime/Engine/Classes/Engine/BlueprintGeneratedClass.h
@@ -3,6 +3,7 @@
 #pragma once
 
 #include "MeshBatch.h"
+#include "ArchiveBase.h"
 #include "BlueprintGeneratedClass.generated.h"
 
 DECLARE_MEMORY_STAT_EXTERN(TEXT("Persistent Uber Graph Frame memory"), STAT_PersistentUberGraphFrameMemory, STATGROUP_Memory, );
@@ -390,6 +391,84 @@ struct ENGINE_API FEventGraphFastCallPair
 	int32 EventGraphCallOffset;
 };
 
+/** A single changed Blueprint component property. */
+USTRUCT()
+struct ENGINE_API FBlueprintComponentChangedPropertyInfo
+{
+	GENERATED_USTRUCT_BODY()
+
+	/** The name of the changed property. */
+	UPROPERTY()
+	FName PropertyName;
+
+	/** The array index of the changed property. */
+	UPROPERTY()
+	int32 ArrayIndex;
+
+	/** The parent struct (owner) of the changed property. */
+	UPROPERTY()
+	UStruct* PropertyScope;
+
+	/** Default constructor. */
+	FBlueprintComponentChangedPropertyInfo()
+	{
+		ArrayIndex = 0;
+		PropertyScope = nullptr;
+	}
+};
+
+/** Cooked data for a Blueprint component template. */
+USTRUCT()
+struct ENGINE_API FBlueprintCookedComponentInstancingData
+{
+	GENERATED_USTRUCT_BODY()
+
+	/** Flag indicating whether or not this contains valid cooked data. Note that an empty changed property list can also be a valid template data context. */
+	UPROPERTY()
+	bool bIsValid;
+
+	/** List of property info records with values that differ between the template and the component class CDO. This list will be generated at cook time. */
+	UPROPERTY()
+	TArray<struct FBlueprintComponentChangedPropertyInfo> ChangedPropertyList;
+
+	/** Source template object name (recorded at load time and used for instancing). */
+	FName ComponentTemplateName;
+
+	/** Source template object class (recorded at load time and used for instancing). */
+	UClass* ComponentTemplateClass;
+
+	/** Source template object flags (recorded at load time and used for instancing). */
+	EObjectFlags ComponentTemplateFlags;
+
+	/** Default constructor. */
+	FBlueprintCookedComponentInstancingData()
+	{
+		bIsValid = false;
+		ComponentTemplateClass = nullptr;
+		ComponentTemplateFlags = RF_NoFlags;
+	}
+
+	/** Builds/returns the internal property list that's used for serialization. This is a linked list of UProperty references. */
+	const FCustomPropertyListNode* GetCachedPropertyListForSerialization() const;
+
+	/** Called at load time to generate the internal cached property data stream from serialization of the source template object. */
+	void LoadCachedPropertyDataForSerialization(UActorComponent* SourceTemplate);
+
+	/** Returns the internal property data stream that's used for fast binary object serialization when instancing components at runtime. */
+	const TArray<uint8>& GetCachedPropertyDataForSerialization() const { return CachedPropertyDataForSerialization; }
+
+protected:
+	/** Internal method used to help recursively build the cached property list for serialization. */
+	void BuildCachedPropertyList(FCustomPropertyListNode** CurrentNode, const UStruct* CurrentScope, int32* CurrentSourceIdx = nullptr) const;
+
+private:
+	/** Internal property list that's used in binary object serialization at component instancing time. */
+	mutable TIndirectArray<FCustomPropertyListNode> CachedPropertyListForSerialization;
+
+	/** Internal property data stream that's used in binary object serialization at component instancing time. */
+	TArray<uint8> CachedPropertyDataForSerialization;
+};
+
 UCLASS()
 class ENGINE_API UBlueprintGeneratedClass : public UClass
 {
@@ -440,6 +519,11 @@ public:
 	UObject* OverridenArchetypeForCDO;
 #endif //WITH_EDITOR
 
+	// Mapping of changed properties & data to apply when instancing components in a cooked build (one entry per named AddComponent node template for fast lookup at runtime).
+	// Note: This is not currently utilized by the editor; it is a runtime optimization for cooked builds only. It assumes that the component class structure does not change.
+	UPROPERTY()
+	TMap<FName, struct FBlueprintCookedComponentInstancingData> CookedComponentInstancingData;
+
 	/** 
 	 * Gets an array of all BPGeneratedClasses (including InClass as 0th element) parents of given generated class 
 	 *
@@ -470,6 +554,7 @@ public:
 	virtual void ConditionalRecompileClass(TArray<UObject*>* ObjLoaded) override;
 	virtual UObject* GetArchetypeForCDO() const override;
 #endif //WITH_EDITOR
+	virtual void SerializeDefaultObject(UObject* Object, FArchive& Ar) override;
 	virtual bool IsFunctionImplementedInBlueprint(FName InFunctionName) const override;
 	virtual uint8* GetPersistentUberGraphFrame(UObject* Obj, UFunction* FuncToCheck) const override;
 	virtual void CreatePersistentUberGraphFrame(UObject* Obj, bool bCreateOnlyIfEmpty = false, bool bSkipSuperClass = false) const override;
@@ -483,8 +568,15 @@ public:
 	{
 		return bHasInstrumentation; 
 	}
+	virtual const FCustomPropertyListNode* GetCustomPropertyListForPostConstruction() const override
+	{
+		return CustomPropertyListForPostConstruction.Num() > 0 ? *CustomPropertyListForPostConstruction.GetData() : nullptr;
+	}
 	// End UClass interface
 
+	/** Called when the custom list of properties used during post-construct initialization needs to be rebuilt (e.g. after serialization and recompilation). */
+	void UpdateCustomPropertyListForPostConstruction();
+
 	static void AddReferencedObjectsInUbergraphFrame(UObject* InThis, FReferenceCollector& Collector);
 
 	static FName GetUberGraphFrameName();
@@ -525,4 +617,12 @@ public:
 			SuperBPClass->InstancePreReplication(Obj, ChangedPropertyTracker);
 		}
 	}
+
+protected:
+	/** Internal helper method used to recursively build the custom property list that's used for post-construct initialization. */
+	void BuildCustomPropertyListForPostConstruction(FCustomPropertyListNode*& InPropertyList, UStruct* InStruct, const uint8* DataPtr, const uint8* DefaultDataPtr);
+
+private:
+	/** List of native class-owned properties that differ from defaults. This is used to optimize property initialization during post-construction by minimizing the number of native class-owned property values that get copied to the new instance. */
+	TIndirectArray<FCustomPropertyListNode> CustomPropertyListForPostConstruction;
 };
diff --git a/Engine/Source/Runtime/Engine/Classes/Engine/InheritableComponentHandler.h b/Engine/Source/Runtime/Engine/Classes/Engine/InheritableComponentHandler.h
--- a/Engine/Source/Runtime/Engine/Classes/Engine/InheritableComponentHandler.h
+++ b/Engine/Source/Runtime/Engine/Classes/Engine/InheritableComponentHandler.h
@@ -71,6 +71,9 @@ struct FComponentOverrideRecord
 	UPROPERTY()
 	FComponentKey ComponentKey;
 
+	UPROPERTY()
+	FBlueprintCookedComponentInstancingData CookedComponentInstancingData;
+
 	FComponentOverrideRecord()
 		: ComponentTemplate(nullptr)
 	{}
@@ -125,6 +128,12 @@ public:
 	FComponentKey FindKey(const FName VariableName) const;
 
 	UActorComponent* GetOverridenComponentTemplate(FComponentKey Key) const;
+	const FBlueprintCookedComponentInstancingData* GetOverridenComponentTemplateData(FComponentKey Key) const;
+
+	TArray<FComponentOverrideRecord>::TIterator CreateRecordIterator()
+	{
+		return Records.CreateIterator();
+	}
 
 private:
 	const FComponentOverrideRecord* FindRecord(const FComponentKey Key) const;
diff --git a/Engine/Source/Runtime/Engine/Classes/Engine/Level.h b/Engine/Source/Runtime/Engine/Classes/Engine/Level.h
--- a/Engine/Source/Runtime/Engine/Classes/Engine/Level.h
+++ b/Engine/Source/Runtime/Engine/Classes/Engine/Level.h
@@ -550,6 +550,7 @@ public:
 #if	WITH_EDITOR
 	virtual void PreEditUndo() override;
 	virtual void PostEditUndo() override;	
+	virtual void BeginCacheForCookedPlatformData(const ITargetPlatform *TargetPlatform) override;
 #endif // WITH_EDITOR
 	virtual void PostLoad() override;
 	virtual void PreSave() override;
diff --git a/Engine/Source/Runtime/Engine/Classes/Engine/SCS_Node.h b/Engine/Source/Runtime/Engine/Classes/Engine/SCS_Node.h
--- a/Engine/Source/Runtime/Engine/Classes/Engine/SCS_Node.h
+++ b/Engine/Source/Runtime/Engine/Classes/Engine/SCS_Node.h
@@ -16,6 +16,10 @@ class USCS_Node : public UObject
 	UPROPERTY()
 	class UActorComponent* ComponentTemplate;
 
+	/** Cached data for faster runtime instancing (only used in cooked builds) */
+	UPROPERTY()
+	FBlueprintCookedComponentInstancingData CookedComponentInstancingData;
+
 	/** If non-None, creates a variable in the class and assigns the created blueprint to it */
 	UPROPERTY()
 	FName VariableName;
@@ -70,12 +74,12 @@ class USCS_Node : public UObject
 	UPROPERTY()
 	bool bVariableNameAutoGenerated_DEPRECATED;
 
-#if WITH_EDITOR
-
 	//~ Begin UObject Interface
 	virtual void PostLoad() override;
 	//~ End UObject Interface
 
+#if WITH_EDITOR
+
 	/** The scene component constructed for component editing in the SCS editor */
 	class USceneComponent* EditorComponentInstance;
 
@@ -97,6 +101,9 @@ class USCS_Node : public UObject
 	/** Return the actual component template used in the BPGC. The template can be overridden in a child. */
 	ENGINE_API UActorComponent* GetActualComponentTemplate(class UBlueprintGeneratedClass* ActualBPGC) const;
 
+	/** Return component template instancing data if cooked for the BPGC, as overridden template data can be cooked out for a child. */
+	ENGINE_API const FBlueprintCookedComponentInstancingData* GetActualComponentTemplateData(class UBlueprintGeneratedClass* ActualBPGC) const;
+
 	/** Returns an array containing this node and all children below it */
 	TArray<USCS_Node*> GetAllNodes();
 	
diff --git a/Engine/Source/Runtime/Engine/Classes/GameFramework/Actor.h b/Engine/Source/Runtime/Engine/Classes/GameFramework/Actor.h
--- a/Engine/Source/Runtime/Engine/Classes/GameFramework/Actor.h
+++ b/Engine/Source/Runtime/Engine/Classes/GameFramework/Actor.h
@@ -2227,6 +2227,7 @@ public:
 
 	/**  Util to create a component based on a template	 */
 	UActorComponent* CreateComponentFromTemplate(UActorComponent* Template, const FName InName = NAME_None );
+	UActorComponent* CreateComponentFromTemplateData(const struct FBlueprintCookedComponentInstancingData* TemplateData, const FName InName = NAME_None);
 
 	DEPRECATED(4.11, "Use CreateComponentFromTemplate that takes a FName instead of a FString")
 	UActorComponent* CreateComponentFromTemplate(UActorComponent* Template, const FString& InName);
@@ -2258,6 +2259,12 @@ protected:
 	*/
 	bool CheckActorComponents();
 
+	/** Checks for and resolve any name conflicts prior to instancing a new Blueprint Component. */
+	void CheckComponentInstanceName(const FName InName);
+
+	/** Called after instancing a new Blueprint Component from either a template or cooked data. */
+	void PostCreateBlueprintComponent(UActorComponent* NewActorComp);
+
 public:
 
 	/** Walk up the attachment chain from RootComponent until we encounter a different actor, and return it. If we are not attached to a component in a different actor, returns NULL */
diff --git a/Engine/Source/Runtime/Engine/Classes/Kismet/DataTableFunctionLibrary.h b/Engine/Source/Runtime/Engine/Classes/Kismet/DataTableFunctionLibrary.h
--- a/Engine/Source/Runtime/Engine/Classes/Kismet/DataTableFunctionLibrary.h
+++ b/Engine/Source/Runtime/Engine/Classes/Kismet/DataTableFunctionLibrary.h
@@ -47,7 +47,8 @@ class ENGINE_API UDataTableFunctionLibrary : public UBlueprintFunctionLibrary
         void* OutRowPtr = Stack.MostRecentPropertyAddress;
         
         P_FINISH;
-        
+		P_NATIVE_BEGIN;
 		*(bool*)RESULT_PARAM = Generic_GetDataTableRowFromName(Table, RowName, OutRowPtr);
+		P_NATIVE_END;
     }
 };
diff --git a/Engine/Source/Runtime/Engine/Classes/Kismet/KismetArrayLibrary.h b/Engine/Source/Runtime/Engine/Classes/Kismet/KismetArrayLibrary.h
--- a/Engine/Source/Runtime/Engine/Classes/Kismet/KismetArrayLibrary.h
+++ b/Engine/Source/Runtime/Engine/Classes/Kismet/KismetArrayLibrary.h
@@ -218,9 +218,9 @@ public:
 		void* NewItemPtr = (Stack.MostRecentPropertyAddress != NULL) ? Stack.MostRecentPropertyAddress : StorageSpace;
  
  		P_FINISH;
- 
+		P_NATIVE_BEGIN;
 		*(int32*)RESULT_PARAM = GenericArray_Add(ArrayAddr, ArrayProperty, NewItemPtr);
- 
+		P_NATIVE_END;
 		InnerProp->DestroyValue(StorageSpace);
 	}
 
@@ -247,9 +247,9 @@ public:
 		void* NewItemPtr = (Stack.MostRecentPropertyAddress != NULL) ? Stack.MostRecentPropertyAddress : StorageSpace;
 
 		P_FINISH;
-
+		P_NATIVE_BEGIN;
 		*(int32*)RESULT_PARAM = GenericArray_AddUnique(ArrayAddr, ArrayProperty, NewItemPtr);
-
+		P_NATIVE_END;
 		InnerProp->DestroyValue(StorageSpace);
 	}
 
@@ -266,8 +266,9 @@ public:
 		}
 
 		P_FINISH;
-
+		P_NATIVE_BEGIN;
 		GenericArray_Shuffle(ArrayAddr, ArrayProperty);
+		P_NATIVE_END;
 	}
 
 	DECLARE_FUNCTION(execArray_Append)
@@ -294,8 +295,9 @@ public:
 		}
 
 		P_FINISH;
-
+		P_NATIVE_BEGIN;
 		GenericArray_Append(TargetArrayAddr, TargetArrayProperty, SourceArrayAddr, SourceArrayProperty);
+		P_NATIVE_END;
 	}
 
 	DECLARE_FUNCTION(execArray_Insert)
@@ -322,9 +324,9 @@ public:
 
 		P_GET_PROPERTY(UIntProperty, Index);
 		P_FINISH;
-
+		P_NATIVE_BEGIN;
 		GenericArray_Insert(ArrayAddr, ArrayProperty, NewItemPtr, Index);
-
+		P_NATIVE_END;
 		InnerProp->DestroyValue(StorageSpace);
 	}
 
@@ -342,8 +344,9 @@ public:
 
 		P_GET_PROPERTY(UIntProperty, Index);
 		P_FINISH;
-
+		P_NATIVE_BEGIN;
 		GenericArray_Remove(ArrayAddr, ArrayProperty, Index);
+		P_NATIVE_END;
 	}
 
 	DECLARE_FUNCTION(execArray_RemoveItem)
@@ -371,9 +374,9 @@ public:
 
 		// Bools need to be processed internally by the property so that C++ bool value is properly set.
 		GenericArray_HandleBool(InnerProp, ItemPtr);
-		  
-		bool WasRemoved = GenericArray_RemoveItem(ArrayAddr, ArrayProperty, ItemPtr);
-		*(bool*)RESULT_PARAM = WasRemoved; 
+		P_NATIVE_BEGIN;
+		*(bool*)RESULT_PARAM = GenericArray_RemoveItem(ArrayAddr, ArrayProperty, ItemPtr);
+		P_NATIVE_END;
 
 		InnerProp->DestroyValue(StorageSpace);
 	}
@@ -390,8 +393,9 @@ public:
 			return;
 		}
 		P_FINISH;
-
+		P_NATIVE_BEGIN;
 		GenericArray_Clear(ArrayAddr, ArrayProperty);
+		P_NATIVE_END;
 	}
 
 	DECLARE_FUNCTION(execArray_Resize)
@@ -407,8 +411,9 @@ public:
 		}
 		P_GET_PROPERTY(UIntProperty, Size);
 		P_FINISH;
-
+		P_NATIVE_BEGIN;
 		GenericArray_Resize(ArrayAddr, ArrayProperty, Size);
+		P_NATIVE_END;
 	}
 
 	DECLARE_FUNCTION(execArray_Length)
@@ -423,8 +428,9 @@ public:
 			return;
 		}
 		P_FINISH;
-
+		P_NATIVE_BEGIN;
 		*(int32*)RESULT_PARAM = GenericArray_Length(ArrayAddr, ArrayProperty);
+		P_NATIVE_END;
 	}
 
 	DECLARE_FUNCTION(execArray_LastIndex)
@@ -439,8 +445,9 @@ public:
 			return;
 		}
 		P_FINISH;
-
+		P_NATIVE_BEGIN;
 		*(int32*)RESULT_PARAM = GenericArray_LastIndex(ArrayAddr, ArrayProperty);
+		P_NATIVE_END;
 	}
 
 	DECLARE_FUNCTION(execArray_Get)
@@ -467,9 +474,9 @@ public:
 		void* ItemPtr = (Stack.MostRecentPropertyAddress != NULL) ? Stack.MostRecentPropertyAddress : StorageSpace;
 
 		P_FINISH;
-
+		P_NATIVE_BEGIN;
 		GenericArray_Get(ArrayAddr, ArrayProperty, Index, ItemPtr);
-
+		P_NATIVE_END;
 		InnerProp->DestroyValue(StorageSpace);
 	}
 
@@ -500,8 +507,9 @@ public:
 
 		P_FINISH;
 
+		P_NATIVE_BEGIN;
 		GenericArray_Set(ArrayAddr, ArrayProperty, Index, NewItemPtr, bSizeToFit);
-
+		P_NATIVE_END;
 		InnerProp->DestroyValue(StorageSpace);
 	}
 
@@ -531,9 +539,10 @@ public:
 		// Bools need to be processed internally by the property so that C++ bool value is properly set.
 		GenericArray_HandleBool(InnerProp, ItemToFindPtr);
 
+		P_NATIVE_BEGIN;
 		// Perform the search
-		int32 FoundIndex = GenericArray_Find(ArrayAddr, ArrayProperty, ItemToFindPtr);
-		*(int32*)RESULT_PARAM = FoundIndex;
+		*(int32*)RESULT_PARAM = GenericArray_Find(ArrayAddr, ArrayProperty, ItemToFindPtr);
+		P_NATIVE_END;
 
 		InnerProp->DestroyValue(StorageSpace);
 	}
@@ -565,8 +574,9 @@ public:
 		GenericArray_HandleBool(InnerProp, ItemToFindPtr);
 
 		// Perform the search
-		int32 FoundIndex = GenericArray_Find(ArrayAddr, ArrayProperty, ItemToFindPtr);
-		*(bool*)RESULT_PARAM = (FoundIndex >= 0);
+		P_NATIVE_BEGIN;
+		*(bool*)RESULT_PARAM = GenericArray_Find(ArrayAddr, ArrayProperty, ItemToFindPtr) >= 0;
+		P_NATIVE_END;
 
 		InnerProp->DestroyValue(StorageSpace);
 	}
@@ -581,6 +591,8 @@ public:
 
 		P_FINISH;
 
+		P_NATIVE_BEGIN;
 		GenericArray_SetArrayPropertyByName(OwnerObject, ArrayPropertyName, SrcArrayAddr);
+		P_NATIVE_END;
 	}
 };
diff --git a/Engine/Source/Runtime/Engine/Classes/Kismet/KismetSystemLibrary.h b/Engine/Source/Runtime/Engine/Classes/Kismet/KismetSystemLibrary.h
--- a/Engine/Source/Runtime/Engine/Classes/Kismet/KismetSystemLibrary.h
+++ b/Engine/Source/Runtime/Engine/Classes/Kismet/KismetSystemLibrary.h
@@ -612,8 +612,9 @@ class ENGINE_API UKismetSystemLibrary : public UBlueprintFunctionLibrary
 		void* SrcStructAddr = Stack.MostRecentPropertyAddress;
 
 		P_FINISH;
-
+		P_NATIVE_BEGIN;
 		Generic_SetStructurePropertyByName(OwnerObject, StructPropertyName, SrcStructAddr);
+		P_NATIVE_END;
 	}
 
 	/** Set a custom structure property by name */
@@ -632,8 +633,9 @@ class ENGINE_API UKismetSystemLibrary : public UBlueprintFunctionLibrary
 		void* SrcStructAddr = Stack.MostRecentPropertyAddress;
 
 		P_FINISH;
-
+		P_NATIVE_BEGIN;
 		Generic_SetStructurePropertyByName(OwnerObject, StructPropertyName, SrcStructAddr);
+		P_NATIVE_END;
 	}
 
 	// --- Collision functions ------------------------------
diff --git a/Engine/Source/Runtime/Engine/Private/ActorConstruction.cpp b/Engine/Source/Runtime/Engine/Private/ActorConstruction.cpp
--- a/Engine/Source/Runtime/Engine/Private/ActorConstruction.cpp
+++ b/Engine/Source/Runtime/Engine/Private/ActorConstruction.cpp
@@ -17,6 +17,8 @@
 
 DEFINE_LOG_CATEGORY(LogBlueprintUserMessages);
 
+DECLARE_CYCLE_STAT(TEXT("InstanceActorComponent"), STAT_InstanceActorComponent, STATGROUP_Engine);
+
 //////////////////////////////////////////////////////////////////////////
 // AActor Blueprint Stuff
 
@@ -576,40 +578,25 @@ UActorComponent* AActor::CreateComponentFromTemplate(UActorComponent* Template,
 	return CreateComponentFromTemplate(Template, FName(*InName));
 }
 
+#if !UE_BUILD_SHIPPING
+static TAutoConsoleVariable<int32> CVarLogBlueprintComponentInstanceCalls(
+	TEXT("LogBlueprintComponentInstanceCalls"),
+	0,
+	TEXT("Log Blueprint Component instance calls; debugging."));
+#endif
+
 UActorComponent* AActor::CreateComponentFromTemplate(UActorComponent* Template, const FName InName)
 {
-	UActorComponent* NewActorComp = NULL;
-	if(Template != NULL)
-	{
-		// If there is a Component with this name already (almost certainly because it is an Instance component), we need to rename it out of the way
-		if (!InName.IsNone())
-		{
-			UObject* ConflictingObject = FindObjectFast<UObject>(this, InName);
-			if (ConflictingObject && ConflictingObject->IsA<UActorComponent>() && CastChecked<UActorComponent>(ConflictingObject)->CreationMethod == EComponentCreationMethod::Instance)
-			{		
-				// Try and pick a good name
-				FString ConflictingObjectName = ConflictingObject->GetName();
-				int32 CharIndex = ConflictingObjectName.Len()-1;
-				while (FChar::IsDigit(ConflictingObjectName[CharIndex]))
-				{
-					--CharIndex;
-				}
-				int32 Counter = 0;
-				if (CharIndex < ConflictingObjectName.Len()-1)
-				{
-					Counter = FCString::Atoi(*ConflictingObjectName.RightChop(CharIndex+1));
-					ConflictingObjectName = ConflictingObjectName.Left(CharIndex+1);
-				}
-				FString NewObjectName;
-				do
-				{
-					NewObjectName = ConflictingObjectName + FString::FromInt(++Counter);
-					
-				} while (FindObjectFast<UObject>(this, *NewObjectName) != nullptr);
+	SCOPE_CYCLE_COUNTER(STAT_InstanceActorComponent);
 
-				ConflictingObject->Rename(*NewObjectName, this);
-			}
-		}
+	UActorComponent* NewActorComp = nullptr;
+	if (Template != nullptr)
+	{
+#if !UE_BUILD_SHIPPING
+		const double StartTime = FPlatformTime::Seconds();
+#endif
+		// Resolve any name conflicts.
+		CheckComponentInstanceName(InName);
 
 		//Make sure, that the name of the instance is different than the name of the template. Otherwise, the template could be handled as an archetype of the instance.
 		const FName NewComponentName = (InName != NAME_None) ? InName : MakeUniqueObjectName(this, Template->GetClass(), Template->GetFName());
@@ -618,10 +605,70 @@ UActorComponent* AActor::CreateComponentFromTemplate(UActorComponent* Template,
 		// Note we aren't copying the the RF_ArchetypeObject flag. Also note the result is non-transactional by default.
 		NewActorComp = (UActorComponent*)StaticDuplicateObject(Template, this, NewComponentName, RF_AllFlags & ~(RF_ArchetypeObject | RF_Transactional | RF_WasLoaded | RF_Public | RF_InheritableComponentTemplate));
 
-		NewActorComp->CreationMethod = EComponentCreationMethod::UserConstructionScript;
+		// Handle post-creation tasks.
+		PostCreateBlueprintComponent(NewActorComp);
 
-		// Need to do this so component gets saved - Components array is not serialized
-		BlueprintCreatedComponents.Add(NewActorComp);
+#if !UE_BUILD_SHIPPING
+		if (CVarLogBlueprintComponentInstanceCalls.GetValueOnGameThread())
+		{
+			UE_LOG(LogBlueprint, Log, TEXT("%s: CreateComponentFromTemplate() - %s \'%s\' completed in %.02g ms"), *GetName(), *Template->GetClass()->GetName(), *NewComponentName.ToString(), (FPlatformTime::Seconds() - StartTime) * 1000.0);
+		}
+#endif
+	}
+	return NewActorComp;
+}
+
+UActorComponent* AActor::CreateComponentFromTemplateData(const FBlueprintCookedComponentInstancingData* TemplateData, const FName InName)
+{
+	SCOPE_CYCLE_COUNTER(STAT_InstanceActorComponent);
+
+	// Component instance data loader implementation.
+	class FBlueprintComponentInstanceDataLoader : public FObjectReader
+	{
+	public:
+		FBlueprintComponentInstanceDataLoader(const TArray<uint8>& InSrcBytes, const FCustomPropertyListNode* InPropertyList)
+			:FObjectReader(const_cast<TArray<uint8>&>(InSrcBytes))
+		{
+			ArCustomPropertyList = InPropertyList;
+			ArUseCustomPropertyList = true;
+			ArWantBinaryPropertySerialization = true;
+		}
+	};
+
+	UActorComponent* NewActorComp = nullptr;
+	if (TemplateData != nullptr)
+	{
+#if !UE_BUILD_SHIPPING
+		const double StartTime = FPlatformTime::Seconds();
+#endif
+		// Resolve any name conflicts.
+		CheckComponentInstanceName(InName);
+
+		//Make sure, that the name of the instance is different than the name of the template. Otherwise, the template could be handled as an archetype of the instance.
+		const FName NewComponentName = (InName != NAME_None) ? InName : MakeUniqueObjectName(this, TemplateData->ComponentTemplateClass, TemplateData->ComponentTemplateName);
+		ensure(NewComponentName != TemplateData->ComponentTemplateName);
+
+		// Note we aren't copying the the RF_ArchetypeObject flag. Also note the result is non-transactional by default.
+		NewActorComp = NewObject<UActorComponent>(
+			this,
+			TemplateData->ComponentTemplateClass,
+			NewComponentName,
+			EObjectFlags(TemplateData->ComponentTemplateFlags) & ~(RF_ArchetypeObject | RF_Transactional | RF_WasLoaded | RF_Public | RF_InheritableComponentTemplate)
+		);
+
+		// Load cached data into the new instance.
+		FBlueprintComponentInstanceDataLoader ComponentInstanceDataLoader(TemplateData->GetCachedPropertyDataForSerialization(), TemplateData->GetCachedPropertyListForSerialization());
+		NewActorComp->Serialize(ComponentInstanceDataLoader);
+
+		// Handle post-creation tasks.
+		PostCreateBlueprintComponent(NewActorComp);
+
+#if !UE_BUILD_SHIPPING
+		if (CVarLogBlueprintComponentInstanceCalls.GetValueOnGameThread())
+		{
+			UE_LOG(LogBlueprint, Log, TEXT("%s: CreateComponentFromTemplateData() - %s \'%s\' completed in %.02g ms"), *GetName(), *TemplateData->ComponentTemplateClass->GetName(), *NewComponentName.ToString(), (FPlatformTime::Seconds() - StartTime) * 1000.0);
+		}
+#endif
 	}
 	return NewActorComp;
 }
@@ -629,13 +676,23 @@ UActorComponent* AActor::CreateComponentFromTemplate(UActorComponent* Template,
 UActorComponent* AActor::AddComponent(FName TemplateName, bool bManualAttachment, const FTransform& RelativeTransform, const UObject* ComponentTemplateContext)
 {
 	UActorComponent* Template = nullptr;
+	FBlueprintCookedComponentInstancingData* TemplateData = nullptr;
 	for (UClass* TemplateOwnerClass = (ComponentTemplateContext != nullptr) ? ComponentTemplateContext->GetClass() : GetClass()
-		; TemplateOwnerClass && !Template
+		; TemplateOwnerClass && !Template && !TemplateData
 		; TemplateOwnerClass = TemplateOwnerClass->GetSuperClass())
 	{
 		if (auto BPGC = Cast<UBlueprintGeneratedClass>(TemplateOwnerClass))
 		{
-			Template = BPGC->FindComponentTemplateByName(TemplateName);
+			// Use cooked instancing data if available (fast path).
+			if (FPlatformProperties::RequiresCookedData())
+			{
+				TemplateData = BPGC->CookedComponentInstancingData.Find(TemplateName);
+			}
+			
+			if (!TemplateData || !TemplateData->bIsValid)
+			{
+				Template = BPGC->FindComponentTemplateByName(TemplateName);
+			}
 		}
 		else if (auto DynamicClass = Cast<UDynamicClass>(TemplateOwnerClass))
 		{
@@ -648,7 +705,7 @@ UActorComponent* AActor::AddComponent(FName TemplateName, bool bManualAttachment
 	}
 
 	bool bIsSceneComponent = false;
-	UActorComponent* NewActorComp = CreateComponentFromTemplate(Template);
+	UActorComponent* NewActorComp = TemplateData ? CreateComponentFromTemplateData(TemplateData) : CreateComponentFromTemplate(Template);
 	if(NewActorComp != nullptr)
 	{
 		// Call function to notify component it has been created
@@ -693,5 +750,48 @@ UActorComponent* AActor::AddComponent(FName TemplateName, bool bManualAttachment
 	return NewActorComp;
 }
 
+void AActor::CheckComponentInstanceName(const FName InName)
+{
+	// If there is a Component with this name already (almost certainly because it is an Instance component), we need to rename it out of the way
+	if (!InName.IsNone())
+	{
+		UObject* ConflictingObject = FindObjectFast<UObject>(this, InName);
+		if (ConflictingObject && ConflictingObject->IsA<UActorComponent>() && CastChecked<UActorComponent>(ConflictingObject)->CreationMethod == EComponentCreationMethod::Instance)
+		{
+			// Try and pick a good name
+			FString ConflictingObjectName = ConflictingObject->GetName();
+			int32 CharIndex = ConflictingObjectName.Len() - 1;
+			while (FChar::IsDigit(ConflictingObjectName[CharIndex]))
+			{
+				--CharIndex;
+			}
+			int32 Counter = 0;
+			if (CharIndex < ConflictingObjectName.Len() - 1)
+			{
+				Counter = FCString::Atoi(*ConflictingObjectName.RightChop(CharIndex + 1));
+				ConflictingObjectName = ConflictingObjectName.Left(CharIndex + 1);
+			}
+			FString NewObjectName;
+			do
+			{
+				NewObjectName = ConflictingObjectName + FString::FromInt(++Counter);
+
+			} while (FindObjectFast<UObject>(this, *NewObjectName) != nullptr);
+
+			ConflictingObject->Rename(*NewObjectName, this);
+		}
+	}
+}
+
+void AActor::PostCreateBlueprintComponent(UActorComponent* NewActorComp)
+{
+	if (NewActorComp)
+	{
+		NewActorComp->CreationMethod = EComponentCreationMethod::UserConstructionScript;
+
+		// Need to do this so component gets saved - Components array is not serialized
+		BlueprintCreatedComponents.Add(NewActorComp);
+	}
+}
 
 
diff --git a/Engine/Source/Runtime/Engine/Private/Blueprint.cpp b/Engine/Source/Runtime/Engine/Private/Blueprint.cpp
--- a/Engine/Source/Runtime/Engine/Private/Blueprint.cpp
+++ b/Engine/Source/Runtime/Engine/Private/Blueprint.cpp
@@ -12,6 +12,7 @@
 #include "Editor/UnrealEd/Public/Kismet2/CompilerResultsLog.h"
 #include "Editor/UnrealEd/Public/Kismet2/StructureEditorUtils.h"
 #include "Editor/Kismet/Public/FindInBlueprintManager.h"
+#include "Editor/UnrealEd/Public/CookerSettings.h"
 #include "Editor/UnrealEd/Public/Editor.h"
 #include "Crc.h"
 #include "MessageLog.h"
@@ -977,6 +978,105 @@ bool UBlueprint::ValidateGeneratedClass(const UClass* InClass)
 	return true;
 }
 
+void UBlueprint::BeginCacheForCookedPlatformData(const ITargetPlatform *TargetPlatform)
+{
+	Super::BeginCacheForCookedPlatformData(TargetPlatform);
+
+	// Only cook component data if the setting is enabled and this is an Actor-based Blueprint class.
+	if (GeneratedClass && GeneratedClass->IsChildOf<AActor>() && GetDefault<UCookerSettings>()->bCookBlueprintComponentTemplateData)
+	{
+		int32 NumCookedComponents = 0;
+		const double StartTime = FPlatformTime::Seconds();
+
+		UBlueprintGeneratedClass* BPGClass = Cast<UBlueprintGeneratedClass>(*GeneratedClass);
+		if (BPGClass)
+		{
+			// Cook all overridden SCS component node templates inherited from the parent class hierarchy.
+			if (UInheritableComponentHandler* TargetInheritableComponentHandler = BPGClass->GetInheritableComponentHandler())
+			{
+				for (auto RecordIt = TargetInheritableComponentHandler->CreateRecordIterator(); RecordIt; ++RecordIt)
+				{
+					if (!RecordIt->CookedComponentInstancingData.bIsValid)
+					{
+						// Note: This will currently block until finished.
+						// @TODO - Make this an async task so we can potentially cook instancing data for multiple components in parallel.
+						FBlueprintEditorUtils::BuildComponentInstancingData(RecordIt->ComponentTemplate, RecordIt->CookedComponentInstancingData);
+						++NumCookedComponents;
+					}
+				}
+			}
+
+			// Cook all SCS component templates that are owned by this class.
+			if (BPGClass->SimpleConstructionScript)
+			{
+				for (auto Node : BPGClass->SimpleConstructionScript->GetAllNodes())
+				{
+					if (!Node->CookedComponentInstancingData.bIsValid)
+					{
+						// Note: This will currently block until finished.
+						// @TODO - Make this an async task so we can potentially cook instancing data for multiple components in parallel.
+						FBlueprintEditorUtils::BuildComponentInstancingData(Node->ComponentTemplate, Node->CookedComponentInstancingData);
+						++NumCookedComponents;
+					}
+				}
+			}
+
+			// Cook all UCS/AddComponent node templates that are owned by this class.
+			for (UActorComponent* ComponentTemplate : BPGClass->ComponentTemplates)
+			{
+				FBlueprintCookedComponentInstancingData& CookedComponentInstancingData = BPGClass->CookedComponentInstancingData.FindOrAdd(ComponentTemplate->GetFName());
+				if (!CookedComponentInstancingData.bIsValid)
+				{
+					// Note: This will currently block until finished.
+					// @TODO - Make this an async task so we can potentially cook instancing data for multiple components in parallel.
+					FBlueprintEditorUtils::BuildComponentInstancingData(ComponentTemplate, CookedComponentInstancingData);
+					++NumCookedComponents;
+				}
+			}
+		}
+
+		if (NumCookedComponents > 0)
+		{
+			UE_LOG(LogBlueprint, Log, TEXT("%s: Cooked %d component(s) in %.02g ms"), *GetName(), NumCookedComponents, (FPlatformTime::Seconds() - StartTime) * 1000.0);
+		}
+	}
+}
+
+bool UBlueprint::IsCachedCookedPlatformDataLoaded(const ITargetPlatform* TargetPlatform)
+{
+	// @TODO - Check async tasks for completion. For now just return TRUE since all tasks will currently block until finished.
+	return true;
+}
+
+void UBlueprint::ClearAllCachedCookedPlatformData()
+{
+	Super::ClearAllCachedCookedPlatformData();
+
+	if (UBlueprintGeneratedClass* BPGClass = Cast<UBlueprintGeneratedClass>(*GeneratedClass))
+	{
+		// Clear cooked data for overridden SCS component node templates inherited from the parent class hierarchy.
+		if (UInheritableComponentHandler* TargetInheritableComponentHandler = BPGClass->GetInheritableComponentHandler())
+		{
+			for (auto RecordIt = TargetInheritableComponentHandler->CreateRecordIterator(); RecordIt; ++RecordIt)
+			{
+				RecordIt->CookedComponentInstancingData = FBlueprintCookedComponentInstancingData();
+			}
+		}
+
+		// Clear cooked data for SCS component node templates.
+		if (BPGClass->SimpleConstructionScript)
+		{
+			for (auto Node : BPGClass->SimpleConstructionScript->GetAllNodes())
+			{
+				Node->CookedComponentInstancingData = FBlueprintCookedComponentInstancingData();
+			}
+		}
+
+		// Clear cooked data for UCS/AddComponent node templates.
+		BPGClass->CookedComponentInstancingData.Empty();
+	}
+}
+
 #endif // WITH_EDITOR
 
 UBlueprint* UBlueprint::GetBlueprintFromClass(const UClass* InClass)
diff --git a/Engine/Source/Runtime/Engine/Private/BlueprintGeneratedClass.cpp b/Engine/Source/Runtime/Engine/Private/BlueprintGeneratedClass.cpp
--- a/Engine/Source/Runtime/Engine/Private/BlueprintGeneratedClass.cpp
+++ b/Engine/Source/Runtime/Engine/Private/BlueprintGeneratedClass.cpp
@@ -64,6 +64,7 @@ void UBlueprintGeneratedClass::PostLoad()
 			CurrObj->MarkPendingKill();
 		}
 	}
+
 #if WITH_EDITORONLY_DATA
 	if (GetLinkerUE4Version() < VER_UE4_CLASS_NOTPLACEABLE_ADDED)
 	{
@@ -75,7 +76,8 @@ void UBlueprintGeneratedClass::PostLoad()
 		}
 	}
 
-	if (const UPackage* Package = GetOutermost())
+	UPackage* Package = GetOutermost();
+	if (Package != nullptr)
 	{
 		if (Package->HasAnyPackageFlags(PKG_ForDiffing))
 		{
@@ -92,6 +94,21 @@ void UBlueprintGeneratedClass::PostLoad()
 		Pair.FunctionToPatch->EventGraphCallOffset = Pair.EventGraphCallOffset;
 	}
 #endif
+
+	// Generate "fast path" instancing data for UCS/AddComponent node templates.
+	if (CookedComponentInstancingData.Num() > 0)
+	{
+		for (int32 Index = ComponentTemplates.Num() - 1; Index >= 0; --Index)
+		{
+			if (UActorComponent* ComponentTemplate = ComponentTemplates[Index])
+			{
+				if (FBlueprintCookedComponentInstancingData* ComponentInstancingData = CookedComponentInstancingData.Find(ComponentTemplate->GetFName()))
+				{
+					ComponentInstancingData->LoadCachedPropertyDataForSerialization(ComponentTemplate);
+				}
+			}
+		}
+	}
 }
 
 void UBlueprintGeneratedClass::GetRequiredPreloadDependencies(TArray<UObject*>& DependenciesOut)
@@ -245,6 +262,95 @@ UObject* UBlueprintGeneratedClass::GetArchetypeForCDO() const
 }
 #endif //WITH_EDITOR
 
+void UBlueprintGeneratedClass::SerializeDefaultObject(UObject* Object, FArchive& Ar)
+{
+	Super::SerializeDefaultObject(Object, Ar);
+
+	if (Ar.IsLoading() && !Ar.IsObjectReferenceCollector() && ClassDefaultObject)
+	{
+		// On load, build the custom property list used in post-construct initialization logic. Note that in the editor, this will be refreshed during compile-on-load.
+		// @TODO - Potentially make this serializable (or cooked data) to eliminate the slight load time cost we'll incur below to generate this list in a cooked build. For now, it's not serialized since the raw UProperty references cannot be saved out.
+		UpdateCustomPropertyListForPostConstruction();
+	}
+}
+
+void UBlueprintGeneratedClass::BuildCustomPropertyListForPostConstruction(FCustomPropertyListNode*& InPropertyList, UStruct* InStruct, const uint8* DataPtr, const uint8* DefaultDataPtr)
+{
+	const UClass* OwnerClass = Cast<UClass>(InStruct);
+	FCustomPropertyListNode** CurrentNodePtr = &InPropertyList;
+
+	for (UProperty* Property = InStruct->PropertyLink; Property; Property = Property->PropertyLinkNext)
+	{
+		const bool bIsConfigProperty = Property->HasAnyPropertyFlags(CPF_Config) && !(OwnerClass && OwnerClass->HasAnyClassFlags(CLASS_PerObjectConfig));
+		const bool bIsTransientProperty = Property->HasAnyPropertyFlags(CPF_Transient | CPF_DuplicateTransient | CPF_NonPIEDuplicateTransient);
+
+		// Skip config properties as they're already in the PostConstructLink chain. Also skip transient properties if they contain a reference to an instanced subobjects (as those should not be initialized from defaults).
+		if (!bIsConfigProperty && (!bIsTransientProperty || !Property->ContainsInstancedObjectProperty()))
+		{
+			for (int32 Idx = 0; Idx < Property->ArrayDim; Idx++)
+			{
+				const uint8* PropertyValue = Property->ContainerPtrToValuePtr<uint8>(DataPtr, Idx);
+				const uint8* DefaultPropertyValue = Property->ContainerPtrToValuePtrForDefaults<uint8>(InStruct, DefaultDataPtr, Idx);
+
+				// If this is a struct property, recurse to pull out any fields that differ from the native CDO.
+				if (UStructProperty* StructProperty = Cast<UStructProperty>(Property))
+				{
+					// Create a new node for the struct property.
+					*CurrentNodePtr = new(CustomPropertyListForPostConstruction) FCustomPropertyListNode(Property, Idx);
+
+					// Recursively gather up all struct fields that differ and assign to the current node's sub property list.
+					BuildCustomPropertyListForPostConstruction((*CurrentNodePtr)->SubPropertyList, StructProperty->Struct, PropertyValue, DefaultPropertyValue);
+
+					// This will be non-NULL if the above found at least one struct field that differs from the native CDO.
+					if ((*CurrentNodePtr)->SubPropertyList)
+					{
+						// Advance to the next node in the list.
+						CurrentNodePtr = &(*CurrentNodePtr)->PropertyListNext;
+					}
+					else
+					{
+						// Remove the node for the struct property since it does not differ from the native CDO.
+						CustomPropertyListForPostConstruction.RemoveAt(CustomPropertyListForPostConstruction.Num() - 1);
+
+						// Clear the current node ptr since the array will have freed up the memory it referenced.
+						*CurrentNodePtr = nullptr;
+					}
+				}
+				else if (!Property->Identical(PropertyValue, DefaultPropertyValue))
+				{
+					// Create a new node, link it into the chain and add it into the array.
+					*CurrentNodePtr = new(CustomPropertyListForPostConstruction) FCustomPropertyListNode(Property, Idx);
+
+					// Advance to the next node ptr.
+					CurrentNodePtr = &(*CurrentNodePtr)->PropertyListNext;
+				}
+			}
+		}
+	}
+}
+
+void UBlueprintGeneratedClass::UpdateCustomPropertyListForPostConstruction()
+{
+	// Empty the current list.
+	CustomPropertyListForPostConstruction.Empty();
+
+	// Find the first native antecedent. All non-native decendant properties are attached to the PostConstructLink chain (see UStruct::Link), so we only need to worry about properties owned by native super classes here.
+	UClass* SuperClass = GetSuperClass();
+	while (SuperClass && !SuperClass->HasAnyClassFlags(CLASS_Native | CLASS_Intrinsic))
+	{
+		SuperClass = SuperClass->GetSuperClass();
+	}
+
+	if (SuperClass)
+	{
+		check(ClassDefaultObject != nullptr);
+
+		// Recursively gather native class-owned property values that differ from defaults.
+		FCustomPropertyListNode* PropertyList = nullptr;
+		BuildCustomPropertyListForPostConstruction(PropertyList, SuperClass, (uint8*)ClassDefaultObject, (uint8*)SuperClass->GetDefaultObject(false));
+	}
+}
+
 bool UBlueprintGeneratedClass::IsFunctionImplementedInBlueprint(FName InFunctionName) const
 {
 	UFunction* Function = FindFunctionByName(InFunctionName);
@@ -890,3 +996,112 @@ void UBlueprintGeneratedClass::GetLifetimeBlueprintReplicationList(TArray<FLifet
 		SuperBPClass->GetLifetimeBlueprintReplicationList(OutLifetimeProps);
 	}
 }
+
+void FBlueprintCookedComponentInstancingData::BuildCachedPropertyList(FCustomPropertyListNode** CurrentNode, const UStruct* CurrentScope, int32* CurrentSourceIdx) const
+{
+	int32 LocalSourceIdx = 0;
+
+	if (CurrentSourceIdx == nullptr)
+	{
+		CurrentSourceIdx = &LocalSourceIdx;
+	}
+
+	// The serialized list is stored linearly, so stop iterating once we no longer match the scope (this indicates that we've finished parsing out "sub" properties for a UStruct).
+	while (*CurrentSourceIdx < ChangedPropertyList.Num() && ChangedPropertyList[*CurrentSourceIdx].PropertyScope == CurrentScope)
+	{
+		// Find changed property by name/scope.
+		const FBlueprintComponentChangedPropertyInfo& ChangedPropertyInfo = ChangedPropertyList[(*CurrentSourceIdx)++];
+		UProperty* Property = nullptr;
+		const UStruct* PropertyScope = CurrentScope;
+		while (!Property && PropertyScope)
+		{
+			Property = FindField<UProperty>(PropertyScope, ChangedPropertyInfo.PropertyName);
+			PropertyScope = PropertyScope->GetSuperStruct();
+		}
+
+		// Create a new node to hold property info.
+		FCustomPropertyListNode* NewNode = new(CachedPropertyListForSerialization) FCustomPropertyListNode(Property, ChangedPropertyInfo.ArrayIndex);
+
+		// Link the new node into the current property list.
+		if (CurrentNode)
+		{
+			*CurrentNode = NewNode;
+		}
+
+		// If this is a UStruct property, recursively build a sub-property list.
+		if (const UStructProperty* StructProperty = Cast<UStructProperty>(Property))
+		{
+			BuildCachedPropertyList(&NewNode->SubPropertyList, StructProperty->Struct, CurrentSourceIdx);
+		}
+
+		// Advance current location to the next linked node.
+		CurrentNode = &NewNode->PropertyListNext;
+	}
+}
+
+const FCustomPropertyListNode* FBlueprintCookedComponentInstancingData::GetCachedPropertyListForSerialization() const
+{
+	FCustomPropertyListNode* PropertyListRootNode = nullptr;
+
+	// Construct the list if necessary.
+	if (CachedPropertyListForSerialization.Num() == 0 && ChangedPropertyList.Num() > 0)
+	{
+		CachedPropertyListForSerialization.Reserve(ChangedPropertyList.Num());
+
+		// Kick off construction of the cached property list.
+		BuildCachedPropertyList(&PropertyListRootNode, ComponentTemplateClass);
+	}
+	else if (CachedPropertyListForSerialization.Num() > 0)
+	{
+		PropertyListRootNode = *CachedPropertyListForSerialization.GetData();
+	}
+
+	return PropertyListRootNode;
+}
+
+void FBlueprintCookedComponentInstancingData::LoadCachedPropertyDataForSerialization(UActorComponent* SourceTemplate)
+{
+	// Blueprint component instance data writer implementation.
+	class FBlueprintComponentInstanceDataWriter : public FObjectWriter
+	{
+	public:
+		FBlueprintComponentInstanceDataWriter(TArray<uint8>& InDstBytes, const FCustomPropertyListNode* InPropertyList)
+			:FObjectWriter(InDstBytes)
+		{
+			ArCustomPropertyList = InPropertyList;
+			ArUseCustomPropertyList = true;
+			ArWantBinaryPropertySerialization = true;
+		}
+	};
+
+	if (bIsValid)
+	{
+		if (SourceTemplate)
+		{
+			// Make sure the source template has been loaded.
+			if (SourceTemplate->HasAnyFlags(RF_NeedLoad))
+			{
+				if (FLinkerLoad* Linker = SourceTemplate->GetLinker())
+				{
+					Linker->Preload(SourceTemplate);
+				}
+			}
+
+			// Cache source template attributes needed for instancing.
+			ComponentTemplateName = SourceTemplate->GetFName();
+			ComponentTemplateClass = SourceTemplate->GetClass();
+			ComponentTemplateFlags = SourceTemplate->GetFlags();
+
+			// This will also load the cached property list, if necessary.
+			const FCustomPropertyListNode* PropertyList = GetCachedPropertyListForSerialization();
+
+			// Write template data out to the "fast path" buffer. All dependencies will be loaded at this point.
+			FBlueprintComponentInstanceDataWriter InstanceDataWriter(CachedPropertyDataForSerialization, PropertyList);
+			SourceTemplate->Serialize(InstanceDataWriter);
+		}
+		else
+		{
+			bIsValid = false;
+		}
+	}
+}
\ No newline at end of file
diff --git a/Engine/Source/Runtime/Engine/Private/InheritableComponentHandler.cpp b/Engine/Source/Runtime/Engine/Private/InheritableComponentHandler.cpp
--- a/Engine/Source/Runtime/Engine/Private/InheritableComponentHandler.cpp
+++ b/Engine/Source/Runtime/Engine/Private/InheritableComponentHandler.cpp
@@ -17,10 +17,18 @@ void UInheritableComponentHandler::PostLoad()
 	for (int32 Index = Records.Num() - 1; Index >= 0; --Index)
 	{
 		FComponentOverrideRecord& Record = Records[Index];
-		if (Record.ComponentTemplate && !CastChecked<UActorComponent>(Record.ComponentTemplate->GetArchetype())->IsEditableWhenInherited())
+		if (Record.ComponentTemplate)
 		{
-			Record.ComponentTemplate->MarkPendingKill(); // hack needed to be able to identify if NewObject returns this back to us in the future
-			Records.RemoveAtSwap(Index);
+			if (!CastChecked<UActorComponent>(Record.ComponentTemplate->GetArchetype())->IsEditableWhenInherited())
+			{
+				Record.ComponentTemplate->MarkPendingKill(); // hack needed to be able to identify if NewObject returns this back to us in the future
+				Records.RemoveAtSwap(Index);
+			}
+			else if(Record.CookedComponentInstancingData.bIsValid)
+			{
+				// Generate "fast path" instancing data.
+				Record.CookedComponentInstancingData.LoadCachedPropertyDataForSerialization(Record.ComponentTemplate);
+			}
 		}
 	}
 }
@@ -339,6 +347,12 @@ UActorComponent* UInheritableComponentHandler::GetOverridenComponentTemplate(FCo
 	return Record ? Record->ComponentTemplate : nullptr;
 }
 
+const FBlueprintCookedComponentInstancingData* UInheritableComponentHandler::GetOverridenComponentTemplateData(FComponentKey Key) const
+{
+	auto Record = FindRecord(Key);
+	return Record ? &Record->CookedComponentInstancingData : nullptr;
+}
+
 const FComponentOverrideRecord* UInheritableComponentHandler::FindRecord(const FComponentKey Key) const
 {
 	for (auto& Record : Records)
diff --git a/Engine/Source/Runtime/Engine/Private/Level.cpp b/Engine/Source/Runtime/Engine/Private/Level.cpp
--- a/Engine/Source/Runtime/Engine/Private/Level.cpp
+++ b/Engine/Source/Runtime/Engine/Private/Level.cpp
@@ -1972,6 +1972,17 @@ void ULevel::OnLevelScriptBlueprintChanged(ULevelScriptBlueprint* InBlueprint)
 	}
 }	
 
+void ULevel::BeginCacheForCookedPlatformData(const ITargetPlatform *TargetPlatform)
+{
+	Super::BeginCacheForCookedPlatformData(TargetPlatform);
+
+	// Cook all level blueprints.
+	for (auto LevelBlueprint : GetLevelBlueprints())
+	{
+		LevelBlueprint->BeginCacheForCookedPlatformData(TargetPlatform);
+	}
+}
+
 #endif	//WITH_EDITOR
 
 bool ULevel::IsPersistentLevel() const
diff --git a/Engine/Source/Runtime/Engine/Private/SCS_Node.cpp b/Engine/Source/Runtime/Engine/Private/SCS_Node.cpp
--- a/Engine/Source/Runtime/Engine/Private/SCS_Node.cpp
+++ b/Engine/Source/Runtime/Engine/Private/SCS_Node.cpp
@@ -32,30 +32,74 @@ UActorComponent* USCS_Node::GetActualComponentTemplate(UBlueprintGeneratedClass*
 	static const FBoolConfigValueHelper EnableInheritableComponents(TEXT("Kismet"), TEXT("bEnableInheritableComponents"), GEngineIni);
 	if (EnableInheritableComponents)
 	{
-		const FComponentKey ComponentKey(this);
-		while (!OverridenComponentTemplate && ActualBPGC)
+		const USimpleConstructionScript* SCS = GetSCS();
+		if (SCS != ActualBPGC->SimpleConstructionScript)
 		{
-			auto InheritableComponentHandler = ActualBPGC->GetInheritableComponentHandler();
-			if (InheritableComponentHandler)
+			const FComponentKey ComponentKey(this);
+			
+			do
 			{
-				OverridenComponentTemplate = InheritableComponentHandler->GetOverridenComponentTemplate(ComponentKey);
-			}
-			ActualBPGC = Cast<UBlueprintGeneratedClass>(ActualBPGC->GetSuperClass());
+				auto InheritableComponentHandler = ActualBPGC->GetInheritableComponentHandler();
+				if (InheritableComponentHandler)
+				{
+					OverridenComponentTemplate = InheritableComponentHandler->GetOverridenComponentTemplate(ComponentKey);
+				}
+
+				ActualBPGC = Cast<UBlueprintGeneratedClass>(ActualBPGC->GetSuperClass());
+
+			} while (!OverridenComponentTemplate && ActualBPGC && SCS != ActualBPGC->SimpleConstructionScript);
 		}
 	}
 	return OverridenComponentTemplate ? OverridenComponentTemplate : ComponentTemplate;
 }
 
+const FBlueprintCookedComponentInstancingData* USCS_Node::GetActualComponentTemplateData(UBlueprintGeneratedClass* ActualBPGC) const
+{
+	const FBlueprintCookedComponentInstancingData* OverridenComponentTemplateData = nullptr;
+	static const FBoolConfigValueHelper EnableInheritableComponents(TEXT("Kismet"), TEXT("bEnableInheritableComponents"), GEngineIni);
+	if (EnableInheritableComponents)
+	{
+		const USimpleConstructionScript* SCS = GetSCS();
+		if (SCS != ActualBPGC->SimpleConstructionScript)
+		{
+			const FComponentKey ComponentKey(this);
+			
+			do
+			{
+				auto InheritableComponentHandler = ActualBPGC->GetInheritableComponentHandler();
+				if (InheritableComponentHandler)
+				{
+					OverridenComponentTemplateData = InheritableComponentHandler->GetOverridenComponentTemplateData(ComponentKey);
+				}
+
+				ActualBPGC = Cast<UBlueprintGeneratedClass>(ActualBPGC->GetSuperClass());
+
+			} while (!OverridenComponentTemplateData && ActualBPGC && SCS != ActualBPGC->SimpleConstructionScript);
+		}
+	}
+
+	return OverridenComponentTemplateData ? OverridenComponentTemplateData : &CookedComponentInstancingData;
+}
+
 UActorComponent* USCS_Node::ExecuteNodeOnActor(AActor* Actor, USceneComponent* ParentComponent, const FTransform* RootTransform, bool bIsDefaultTransform)
 {
 	check(Actor != nullptr);
 	check((ParentComponent != nullptr && !ParentComponent->IsPendingKill()) || (RootTransform != nullptr)); // must specify either a parent component or a world transform
 
+	// Create a new component instance based on the template
+	UActorComponent* NewActorComp = nullptr;
 	auto ActualBPGC = Cast<UBlueprintGeneratedClass>(Actor->GetClass());
-	UActorComponent* ActualComponentTemplate = GetActualComponentTemplate(ActualBPGC);
+	auto ActualComponentTemplateData = FPlatformProperties::RequiresCookedData() ? GetActualComponentTemplateData(ActualBPGC) : nullptr;
+	if (ActualComponentTemplateData && ActualComponentTemplateData->bIsValid)
+	{
+		// Use cooked instancing data if valid (fast path).
+		NewActorComp = Actor->CreateComponentFromTemplateData(ActualComponentTemplateData, VariableName);
+	}
+	else if (UActorComponent* ActualComponentTemplate = GetActualComponentTemplate(ActualBPGC))
+	{
+		NewActorComp = Actor->CreateComponentFromTemplate(ActualComponentTemplate, VariableName);
+	}
 
-	// Create a new component instance based on the template
-	UActorComponent* NewActorComp = Actor->CreateComponentFromTemplate(ActualComponentTemplate, VariableName);
 	if(NewActorComp != nullptr)
 	{
 		NewActorComp->CreationMethod = EComponentCreationMethod::SimpleConstructionScript;
@@ -351,6 +395,21 @@ void USCS_Node::RemoveMetaData(const FName& Key)
 	}
 }
 
+void USCS_Node::PostLoad()
+{
+	Super::PostLoad();
+
+#if WITH_EDITOR
+	ValidateGuid();
+#endif
+
+	// If valid, load cooked component instancing data.
+	if (ComponentTemplate && CookedComponentInstancingData.bIsValid)
+	{
+		CookedComponentInstancingData.LoadCachedPropertyDataForSerialization(ComponentTemplate);
+	}
+}
+
 #if WITH_EDITOR
 void USCS_Node::SetParent(USCS_Node* InParentNode)
 {
@@ -460,13 +519,6 @@ USceneComponent* USCS_Node::GetParentComponentTemplate(UBlueprint* InBlueprint)
 	return ParentComponentTemplate;
 }
 
-void USCS_Node::PostLoad()
-{
-	Super::PostLoad();
-
-	ValidateGuid();
-}
-
 void USCS_Node::ValidateGuid()
 {
 	// Backward compatibility:
diff --git a/Engine/Source/Runtime/Engine/Private/SkeletalMesh.cpp b/Engine/Source/Runtime/Engine/Private/SkeletalMesh.cpp
--- a/Engine/Source/Runtime/Engine/Private/SkeletalMesh.cpp
+++ b/Engine/Source/Runtime/Engine/Private/SkeletalMesh.cpp
@@ -5045,14 +5045,6 @@ void USkeletalMeshComponent::Serialize(FArchive& Ar)
 	}
 #endif
 
-	if (Ar.IsSaving() && Ar.UE4Ver() < VER_UE4_NO_ANIM_BP_CLASS_IN_GAMEPLAY_CODE)
-	{
-		if ((nullptr != AnimBlueprintGeneratedClass) && (nullptr == AnimClass))
-		{
-			AnimClass = AnimBlueprintGeneratedClass;
-		}
-	}
-
 	Super::Serialize(Ar);
 			
 	// to count memory : TODO: REMOVE?
@@ -5090,9 +5082,9 @@ void USkeletalMeshComponent::Serialize(FArchive& Ar)
 	}
 #endif
 
-	if (Ar.IsLoading() && Ar.UE4Ver() < VER_UE4_NO_ANIM_BP_CLASS_IN_GAMEPLAY_CODE)
+	if (Ar.IsLoading() && (Ar.UE4Ver() < VER_UE4_NO_ANIM_BP_CLASS_IN_GAMEPLAY_CODE))
 	{
-		if ((nullptr != AnimBlueprintGeneratedClass) && (nullptr == AnimClass))
+		if (nullptr != AnimBlueprintGeneratedClass)
 		{
 			AnimClass = AnimBlueprintGeneratedClass;
 		}
diff --git a/Engine/Source/Runtime/Engine/Public/BlueprintUtilities.h b/Engine/Source/Runtime/Engine/Public/BlueprintUtilities.h
--- a/Engine/Source/Runtime/Engine/Public/BlueprintUtilities.h
+++ b/Engine/Source/Runtime/Engine/Public/BlueprintUtilities.h
@@ -3,6 +3,7 @@
 #pragma once
 
 #include "GraphEditAction.h"
+#include "Runtime/Core/Public/Serialization/Archive.h"
 
 DECLARE_MULTICAST_DELEGATE_OneParam( FOnGraphChanged, const FEdGraphEditAction& );
 DECLARE_DELEGATE_OneParam( FSingleNodeEvent, class UEdGraphNode* );
diff --git a/Engine/Source/Runtime/Engine/Public/GeneratedCodeHelpers.h b/Engine/Source/Runtime/Engine/Public/GeneratedCodeHelpers.h
--- a/Engine/Source/Runtime/Engine/Public/GeneratedCodeHelpers.h
+++ b/Engine/Source/Runtime/Engine/Public/GeneratedCodeHelpers.h
@@ -6,6 +6,7 @@
 // THIS FILE SHOULD BE USED ONLY BY AUTOMATICALLY GENERATED CODE. 
 
 // Common includes
+#include "Runtime/Core/Public/Core.h"
 #include "UObject/Stack.h"
 #include "Blueprint/BlueprintSupport.h"
 #include "Engine/BlueprintGeneratedClass.h"
@@ -95,7 +96,7 @@ public:
 	}
 
 	template<typename T>
-	static void Array_Shuffle(TArray<T>& TargetArray, const UArrayProperty* ArrayProperty)
+	static void Array_Shuffle(TArray<T>& TargetArray)
 	{
 		int32 LastIndex = TargetArray.Num() - 1;
 		for (int32 i = 0; i < LastIndex; ++i)
@@ -376,23 +377,23 @@ struct FUnconvertedWrapper
 template<typename T>
 struct TArrayCaster
 {
-	TArray<T*> Val;
-	TArray<T*>& ValRef;
+	TArray<T> Val;
+	TArray<T>& ValRef;
 
-	TArrayCaster(TArray<T*>& InArr) 
+	TArrayCaster(const TArray<T>& InArr)
 		: Val()
-		, ValRef(InArr)
+		, ValRef(*(TArray<T>*)(&InArr))
 	{}
 
-	TArrayCaster(TArray<T*>&& InArr) 
+	TArrayCaster(TArray<T>&& InArr) 
 		: Val(MoveTemp(InArr))
 		, ValRef(Val)
 	{}
 
 	template<typename U>
-	TArray<U*>& Get()
+	TArray<U>& Get()
 	{
-		static_assert(sizeof(T*) == sizeof(U*), "Incompatible pointers");
-		return *reinterpret_cast<TArray<U*>*>(&ValRef);
+		static_assert(sizeof(T) == sizeof(U), "Incompatible pointers");
+		return *reinterpret_cast<TArray<U>*>(&ValRef);
 	}
 };
\ No newline at end of file
diff --git a/Engine/Source/Runtime/GameplayAbilities/Public/GameplayEffectUIData_TextOnly.h b/Engine/Source/Runtime/GameplayAbilities/Public/GameplayEffectUIData_TextOnly.h
--- a/Engine/Source/Runtime/GameplayAbilities/Public/GameplayEffectUIData_TextOnly.h
+++ b/Engine/Source/Runtime/GameplayAbilities/Public/GameplayEffectUIData_TextOnly.h
@@ -1,7 +1,7 @@
 // Copyright 1998-2016 Epic Games, Inc. All Rights Reserved.
 
 #pragma once
-
+#include "GameplayEffectUIData.h"
 #include "GameplayEffectUIData_TextOnly.generated.h"
 
 /**
diff --git a/Engine/Source/Runtime/MovieSceneTracks/Public/Sections/MovieSceneColorSection.h b/Engine/Source/Runtime/MovieSceneTracks/Public/Sections/MovieSceneColorSection.h
--- a/Engine/Source/Runtime/MovieSceneTracks/Public/Sections/MovieSceneColorSection.h
+++ b/Engine/Source/Runtime/MovieSceneTracks/Public/Sections/MovieSceneColorSection.h
@@ -1,7 +1,8 @@
 // Copyright 1998-2016 Epic Games, Inc. All Rights Reserved.
 
 #pragma once
-
+#include "MovieSceneSection.h"
+#include "IKeyframeSection.h"
 #include "MovieSceneColorSection.generated.h"
 
 enum class EKeyColorChannel
diff --git a/Engine/Source/Runtime/MovieSceneTracks/Public/Sections/MovieSceneFloatSection.h b/Engine/Source/Runtime/MovieSceneTracks/Public/Sections/MovieSceneFloatSection.h
--- a/Engine/Source/Runtime/MovieSceneTracks/Public/Sections/MovieSceneFloatSection.h
+++ b/Engine/Source/Runtime/MovieSceneTracks/Public/Sections/MovieSceneFloatSection.h
@@ -2,6 +2,8 @@
 
 #pragma once
 
+#include "MovieSceneSection.h"
+#include "IKeyframeSection.h"
 #include "MovieSceneFloatSection.generated.h"
 
 /**
diff --git a/Engine/Source/Runtime/UMG/Public/UMG.h b/Engine/Source/Runtime/UMG/Public/UMG.h
--- a/Engine/Source/Runtime/UMG/Public/UMG.h
+++ b/Engine/Source/Runtime/UMG/Public/UMG.h
@@ -134,3 +134,12 @@
 
 // Interfaces
 #include "IUMGModule.h"
+
+// Animation 
+#include "MovieScene2DTransformSection.h"
+#include "MovieScene2DTransformTrack.h"
+#include "MovieSceneMarginSection.h"
+#include "MovieSceneMarginTrack.h"
+#include "UMGSequencePlayer.h"
+#include "WidgetAnimation.h"
+#include "WidgetAnimationBinding.h"
\ No newline at end of file
