Reimplementing CL #3005441 by [REDACTED_PERSON] from Dev-Blueprints to fix UE-32012 (Nativized package project crashes on launch)

  AssembleReferenceTokenStream is called for Dynamic Classes:
  - in ConstructDynamicType() (when class is explicitly loaded)
  - in __CustomDynamicClassInitialization() (when CDO is created)

diff --git a/Engine/Source/Developer/BlueprintCompilerCppBackend/Private/BlueprintCompilerCppBackendValueHelper.cpp b/Engine/Source/Developer/BlueprintCompilerCppBackend/Private/BlueprintCompilerCppBackendValueHelper.cpp
--- a/Engine/Source/Developer/BlueprintCompilerCppBackend/Private/BlueprintCompilerCppBackendValueHelper.cpp
+++ b/Engine/Source/Developer/BlueprintCompilerCppBackend/Private/BlueprintCompilerCppBackendValueHelper.cpp
@@ -791,6 +791,7 @@ void FEmitDefaultValueHelper::GenerateCustomDynamicClassInitialization(FEmitterL
 	Context.AddLine(TEXT("ensure(0 == InDynamicClass->ComponentTemplates.Num());"));
 	Context.AddLine(TEXT("ensure(0 == InDynamicClass->Timelines.Num());"));
 	Context.AddLine(TEXT("ensure(nullptr == InDynamicClass->AnimClassImplementation);"));
+	Context.AddLine(TEXT("InDynamicClass->AssembleReferenceTokenStream();"));
 
 	Context.CurrentCodeType = FEmitterLocalContext::EGeneratedCodeType::SubobjectsOfClass;
 	Context.ResetPropertiesForInaccessibleStructs();
diff --git a/Engine/Source/Runtime/CoreUObject/Private/UObject/UObjectBase.cpp b/Engine/Source/Runtime/CoreUObject/Private/UObject/UObjectBase.cpp
--- a/Engine/Source/Runtime/CoreUObject/Private/UObject/UObjectBase.cpp
+++ b/Engine/Source/Runtime/CoreUObject/Private/UObject/UObjectBase.cpp
@@ -1157,7 +1157,10 @@ UObject* ConstructDynamicType(FName TypePathName)
 	UObject* Result = nullptr;
 	if (FClassConstructFunctions* ClassConstructFn = GetDynamicClassMap().Find(TypePathName))
 	{
-		Result = ClassConstructFn->StaticClassFn();
+		UClass* DynamicClass = ClassConstructFn->StaticClassFn();
+		check(DynamicClass);
+		DynamicClass->AssembleReferenceTokenStream();
+		Result = DynamicClass;
 	}
 	else if (UScriptStruct *(**StaticStructFNPtr)() = GetDynamicStructMap().Find(TypePathName))
 	{
