Copying //UE4/Dev-Blueprints to //UE4/Dev-Main (Source: //UE4/Dev-Blueprints @ 3025888)

==========================
MAJOR FEATURES + CHANGES
==========================

Change 2927746 on 2016/03/30 by [REDACTED_PERSON]

	Local variables in function graphs will now store a hard reference to their UObject value.

	Fixes a crash when a Blueprint is saved before compiling with the local variable's value set. Ensures that the UObject is loaded with the Blueprint.

Change 2927751 on 2016/03/30 by [REDACTED_PERSON]

	Back out changelist 2927746

Change 2986483 on 2016/05/23 by [REDACTED_PERSON]

	Added overriden names of Enum keys.

Change 2986712 on 2016/05/23 by [REDACTED_PERSON]

	[UE-21010] Apply updated transform to component template instances when changing the scene root in a Blueprint class.

	change summary:
	- modified SSCS_RowWidget::OnMakeNewRootDropAction() to propagate the location/rotation reset to instances of the component template that's becoming the new scene root.

Change 2987406 on 2016/05/23 by [REDACTED_PERSON]

	Fixed Functions filter in Find-In-Blueprints will show components from the SCS

Change 2988925 on 2016/05/24 by [REDACTED_PERSON]

	Fixed Issue where certain primitives would not automatically type cast to Text in Blueprint graph.

Change 2989001 on 2016/05/24 by [REDACTED_PERSON]

	PR #2418: Fixed a typo in Blueprint.h (Contributed by [REDACTED_PERSON])

Change 2989447 on 2016/05/25 by [REDACTED_PERSON]

	[UE-30807] Propagate edit condition property value changes to instances of template objects.

	change summary:
	- modified FPropertyEditor::SetEditConditionState() to propagate an EditConditionProperty value change to all instances if the outer owning object is a template (e.g. CDO)

Change 2989804 on 2016/05/25 by [REDACTED_PERSON]

	[UE-30289] Preserve relative scale on the root scene component when converting an Actor instance to a Blueprint Class.

	change summary:
	- modified FKismetEditorUtilities::CreateBlueprintFromActor() to post-copy the relative scale value from the Actor's root component to the new Blueprint CDO's root component

Change 2990234 on 2016/05/25 by [REDACTED_PERSON]

	Fixed issue where including a period ina  Blueprint function causes double-click to fail to open its graph

Change 2990566 on 2016/05/25 by [REDACTED_PERSON]

	Better warn logging to help locate variable nodes that emit a "variable not found" message.

Change 2991083 on 2016/05/26 by [REDACTED_PERSON]

	Blueprint nativization: converted classes have "config" specified.

Change 2991363 on 2016/05/26 by [REDACTED_PERSON]

	[UE-19599] Copy-and-paste of Actor instances from level to Blueprint/IWCE component tree views now adds properly-initialized components.

	change summary:
	- modified FCustomizableTextObjectFactory::CanCreateObjectsFromText() to handle "Begin Actor/End Actor" blocks in T3D text
	- modified FCustomizableTextObjectFactory::ProcessBuffer() to handle "Begin Actor/End Actor" blocks in T3D text (so that Actor-type objects can be processed)
	- modified FComponentObjectTextFactory::CanCreateClass() to allow Actor-type objects to pass
	- modified FComponentObjectTextFactory::ProcessConstructedObject() to handle Actor-type objects and pull out owned component instances as constructed objects

Change 2992990 on 2016/05/27 by [REDACTED_PERSON]

	Fixed issue where Connecting Self Reference Pin to a String pin does not fully connect the generated GetDisplayName node

Change 2992995 on 2016/05/27 by [REDACTED_PERSON]

	Fixed issue where GetClass node is not listed in the Context Menu when pulling from a self node and Context Sensitive is checked.

Change 2993449 on 2016/05/27 by [REDACTED_PERSON]

	[UE-31379] Don't instrument "preview" Actor instances during Blueprint profiler script event processing.

	change summary:
	- modified FBlueprintProfiler::InstrumentEvent() to check for and bypass Actor instances belonging to a preview or inactive world type.

Change 2993531 on 2016/05/27 by [REDACTED_PERSON]

	PR #2433: Interface functions inherited from a native base class now appear in . (Contributed by [REDACTED_PERSON])

Change 2993969 on 2016/05/30 by [REDACTED_PERSON]

	UE-30729 Crash in Native Orion when selecting Sword or Tomahawk

	Clear AsyncLoading in subobjects.

Change 2993990 on 2016/05/30 by [REDACTED_PERSON]

	[UE-30984] Exclude reroute nodes from Blueprint profiler node mapping.

	change summary:
	- modified FBlueprintFunctionContext::MapInputPins() to pass through non-relevant nodes when iterating through non-exec input pin links.
	- modified FBlueprintFunctionContext::MapExecPins() to pass through non-relevant nodes when iterating through output exec pin links.
	- modified FBlueprintFunctionContext::MapTunnelEntry() to pass through non-relevant nodes when iterating through tunnel node exit points.
	- modified FBlueprintFunctionContext::MapTunnelInstance() to pass through non-relevant nodes when iterating through tunnel graph entry points.

Change 2994591 on 2016/05/31 by [REDACTED_PERSON]

	Fixed issue where inherited Blueprint variable would not show parent's replications settings

Change 2994613 on 2016/05/31 by [REDACTED_PERSON]

	Minor refactor and Various fixes to the blueprint profiler moving towards MVP goal.

	Notes:-
	- Sequence node inclusive time fixed
	- Trace History tidy up
	- Compile Icon and status messages for instrumentation
	- Message in the profiler tab for instrumentation
	- Profiler view tidy up and heat thresholds controls added
	- fixed the summed execution branch stats
	- fixed the connection drawing policy to use branch pin stats and fixed the crash from UE-31396
	- added hottest path and hottest endpoint wire heatmaps
	- switched off the graph filter by default
	- added total time for the heatmaps
	- fixed issue where initialising mapped functions caused an assert due to changes to the array/map in initialisation code

Change 2995058 on 2016/05/31 by [REDACTED_PERSON]

	[UE-30718] Native/const implementable events will no longer cause a crash at runtime when the Blueprint profiler is running.

	change summary:
	- modified UObject::ProcessEvent() to bypass instrumentation for native event functions that are not implemented (overridden) in a BP class.
	- modified FScriptEventPlayback::Process() to first check for a standalone function match (UCS, implementable events declared as 'const') before settling on the ubergraph function for the target context.

Change 2995218 on 2016/05/31 by [REDACTED_PERSON]

	[UE-30778] Restored non-K2 compact graph nodes (e.g. Material Editor) to previous size.

	change summary:
	- modified SGraphNode::GetNodeIndicatorOverlayVisibility() default impl to return 'Collapsed' by default, so it doesn't affect layout.

Change 2996417 on 2016/06/01 by [REDACTED_PERSON]

	[UE-16073] Basic shape components (cube etc.) will now apply the correct override material to instances after being added through the component tree in the Blueprint editor.

	change summary:
	- modified the 'OnBasicShapeCreated' lambda in FComponentTypeRegistryData::AddBasicShapeComponents() to propagate the material override to all instances when the given component is an archetype (template) object.

Change 2997001 on 2016/06/01 by [REDACTED_PERSON]

	Fixed Double Clicking a component in the results of Find-In-Blueprints does not select the component

Change 2997521 on 2016/06/02 by [REDACTED_PERSON]

	[Blueprint Nativization]
	- Added FilesToIncludeInModuleHeader config variable in BlueprintNativizationSettings. So some headers can be included in NativizedAssets.h
	- Guids of nodes are no longer recreated when Blueprint is duplicated for "C++ compilation". Previously child bp used variable names based on original parent class, but nativized parent class had guids recreated.

Change 2997522 on 2016/06/02 by [REDACTED_PERSON]

	Native implementation of NOEXPORT FInterpCurvePoint structures. (It's necessary for Blueprint nativization)

Change 2997638 on 2016/06/02 by [REDACTED_PERSON]

	Improvements for Blueprint Nativization:
	- Overridden names in nativized code have proper escape characters (in generated code).
	- OnlyDefaultConstructorDeclared metadata is replaced by ObjectInitializerConstructorDeclared
	- Arrays of nativized anum have the following form: TArray<Enum> (previously it was TArray<TEnumAsByte<Enum>>)
	- warning C4883 is disabled  in .generated.cpp files for nativized module

Change 2997639 on 2016/06/02 by [REDACTED_PERSON]

	Minor improvements in Ocean gameplay code. Required for Blueprint Nativization.

Change 2997656 on 2016/06/02 by [REDACTED_PERSON]

	Various improvements in BlueprintCompilerCppBackend:
	- Fixed interface cast
	- Fixed TSwitchValue issue (when used with literals)
	- Fixed improper name for NativeBlueprintEvent (when calling parent's implementation)
	- Fixed bitfield getter code.
	- Reduce code size (less UsedAssets, less  ReferencedConvertedFields, cached UEnums)
	- operator == is generated for nativized structs
	- Fixed AssedId (AssetPtr) constructor in nativized code.
	- Fixed arrays of noexport struct
	- Fixed missing headers for native single cast delegate signature.
	- Fixed issue when default constructor (in native) is missing (constructor with FObjectInitialized, wont be used automatically). See "ObjectInitializerConstructorDeclared" metadata.

Change 2997691 on 2016/06/02 by [REDACTED_PERSON]

	operator == in FText. It is required for some functions in TArray<FText>

Change 2997793 on 2016/06/02 by [REDACTED_PERSON]

	Added support for BaseAsyncTask nodes, fixed a problem with instance mapping and turned off the debug instance filter
	#Proj BlueprintProfiler, Kismet

Change 2997901 on 2016/06/02 by [REDACTED_PERSON]

	Back out changelist 2997691

Change 2998038 on 2016/06/02 by [REDACTED_PERSON]

	Merging //UE4/Dev-Main to Dev-Blueprints (//UE4/Dev-Blueprints)

Change 2998052 on 2016/06/02 by [REDACTED_PERSON]

	Fixed Comment bubbles not remembering changes after losing focus

Change 2998450 on 2016/06/02 by [REDACTED_PERSON]

	[UE-31550] Fix crash on load of a Blueprint class containing a bitmask variable with missing enum type metadata.

	change summary:
	- modified FBlueprintEditorUtils::ValidateBlueprintVariableMetadata() to check for presence of bitmask enum type metadata on a variable before trying to validate it.

Change 2999763 on 2016/06/03 by [REDACTED_PERSON]

	Guarding against a crash with an ensure - attempting to catch why this is happening by logging more info, as we're unable to repro it. Guarding against nodes which reference malformed (TRASH) classes.

Change 2999768 on 2016/06/03 by [REDACTED_PERSON]

	This is just workaound. FReferenceFinder::FindReferences doesn;t find Enum variable in UByteProperty.

Change 2999770 on 2016/06/03 by [REDACTED_PERSON]

	[Blueprint Nativization]
	Workaround for missing ==operator in native structures. The generated code uses special version of array funtions.

Change 2999798 on 2016/06/03 by [REDACTED_PERSON]

	Guarding against malformed Blueprints  (ones without valid "authoratative" class) used as context for the node menu. Baffling how we'd get into this scenario, but this adds ensures to hopefully give us clues and stabalize the editor.

Change 2999941 on 2016/06/03 by [REDACTED_PERSON]

	Correcting mistake in previously attempted fix (CL 2781229). Now using weak ptr IsValid checks to guard against destroyed nodes in deferred graph actions (TWeakObjectPtr::Get() does not check IsValid before returning).

Change 3001731 on 2016/06/06 by [REDACTED_PERSON]

	[UE-30638] BP profiler will no longer crash at runtime while profiling events that call functions on an external target.

	change summary:
	- modified FBlueprintProfiler::ProcessEventProfilingData() to only remove 'Class' and 'Instance' signals on new events.
	- modified FScriptEventPlayback::NodeSignalHelper struct to include a new 'BlueprintContext' field.
	- modified FScriptEventPlayback::Process() to handle midstream context switches by updating the Blueprint/Function context on 'Class' and/or 'Instance' signals.
	- modified FScriptEventPlayback::Process() to cache and reference the current Blueprint context within the cached NodeSignalHelper while handling processed events.

Change 3002075 on 2016/06/06 by [REDACTED_PERSON]

	Improved FScriptBuilderBase::EmitTermExpr in KismetCompilerVMBackend.
	Literal expression can be emitted without known desitination property.

Change 3002096 on 2016/06/06 by [REDACTED_PERSON]

	This change expands the way that the blueprint profiler detects event nodes during mapping to include other non function graphs.
	#Proj BlueprintProfiler

Change 3002108 on 2016/06/06 by [REDACTED_PERSON]

	Adds a new default option to average the blueprint level stats in the profiler.
	#Proj Kismet, BlueprintProfiler

	- The controls were also getting a bit messy so I tidied them all up into a re-usable toolbar for convenience going forward.

Change 3002782 on 2016/06/06 by samuel.proctor

	Test assets for Interface testing

Change 3003826 on 2016/06/07 by [REDACTED_PERSON]

	A few minor visual improvements for the blueprint profiler.
	#Proj Kismet, BlueprintProfiler, EditorStyle

	- Updated the actor icon to match the world outliner and added some functionality to draw attention to stale/deleted actors.
	- Updated the pure node icon.

Change 3004067 on 2016/06/07 by samuel.proctor

	New test asset for blueprint interfaces

Change 3004069 on 2016/06/07 by samuel.proctor

	Updating asset for Interface testing

Change 3004275 on 2016/06/07 by [REDACTED_PERSON]

	Fixed issue where Toggle Comment Bubble button for Reroute nodes would not rever tthe comment bubble to constant visibility

Change 3004329 on 2016/06/07 by [REDACTED_PERSON]

	EdGraphPin is no longer a UObject, this will improve load times significantly on projects with large number of blueprints, but content does need to be resaved in order to see the improvement in load time. UObject counts are also greatly reduced.

Change 3004418 on 2016/06/07 by [REDACTED_PERSON]

	KismetCompilerVMBackend:  Fixed issue, when a byte property has no enum specified (for examle parameter from EqualEqual_ByteByte) but the enum is needed to parse a literal value.

Change 3004496 on 2016/06/07 by [REDACTED_PERSON]

	Disabling expensive pin allocation tracking

Change 3004649 on 2016/06/07 by [REDACTED_PERSON]

	Preventing a new warning from being generated on trace point exceptions (trace point exceptions are used to hook into the debugger, and don't represent errors).

Change 3004667 on 2016/06/07 by [REDACTED_PERSON]

	Removed my debugging logic

Change 3004848 on 2016/06/07 by [REDACTED_PERSON]

	Fix spammy ensure

Change 3004871 on 2016/06/07 by [REDACTED_PERSON]

	[UE-24950] No longer including components instanced as default subobjects of and attached to components instanced by construction script in the IWCE component tree view.

	change summary:
	- modified SSCSEditor::UpdateTree() to exclude child components instanced in native code as "nested" DSOs and parented to non-natively-constructed (e.g. Blueprint) components; these instances are no longer being shown in IWCE in order to avoid confusion, as they're not currently mutable at the instance level, will always be parented to something that is visible in the tree, and they're also not currently shown in the Blueprint editor's component tree view (because they're not stored in the CDO).
	- modified FSceneComponentData's ctor to exclude child components instanced in native code as nested DSOs from the AttachedInstancedComponents array; this allows child components instanced as nested DSOs to be disposed of along with the constructed parent instance when re-running construction scripts.

Change 3005203 on 2016/06/07 by [REDACTED_PERSON]

	Fix for undo/redo/serialization issues with ed graph pin change. When serialization logic was applied incrementally our attempts to keep LinkedTo symetrical and aggressively clear destroyed nodes caused problems

Change 3005441 on 2016/06/08 by [REDACTED_PERSON]

	AssembleReferenceTokenStream is called for Dynamic Classes:
	- in ConstructDynamicType() (when class is explicitly loaded)
	- in __CustomDynamicClassInitialization() (when CDO is created)

Change 3005540 on 2016/06/08 by [REDACTED_PERSON]

	This adds the ability to track profiler instances between editor and PIE instances and displays the current status through the icon coloring.
	#Proj BlueprintProfiler, Kismet

	- The jira was already fixed but I think this change improves the instance status clarity

Change 3006196 on 2016/06/08 by [REDACTED_PERSON]

	Copy/paste logic for pin connections got lost in the shuffle

Change 3006416 on 2016/06/08 by [REDACTED_PERSON]

	[UE-31735] Fix potential loss of GetClassDefaults node output pin links on load (due to dependency load order).

	change summary:
	- modified UK2Node_GetClassDefaults::GetInputClass() to redirect to the generated skeleton class only if it's valid. this ensures that output pins will be reallocated during node reconstruction even if the dependent Blueprint's skeleton class has not yet been generated on load.

Change 3006522 on 2016/06/08 by [REDACTED_PERSON]

	Under rare circumstances a deprecated pin comes in that is outered to the transient package

Change 3006576 on 2016/06/08 by [REDACTED_PERSON]

	Fix for non-editor builds

Change 3006610 on 2016/06/08 by [REDACTED_PERSON]

	[UE-31743] Fix data loss issue when loading a serialized non-native component class instance that's owned by an Actor-based Blueprint class instance.

	change summary:
	- modified FObjectInitializer::InitProperties() to disable fast path initialization for non-native class types when the default data does not equate to the non-native CDO (as is also done within the native path). this is necessary because the optimized property list that we generate at load time to support fast path initialization of Blueprint class instances is only applicable to the generated CDO.

Change 3006824 on 2016/06/08 by [REDACTED_PERSON]

	More undo/redo fixes, this time fixes for when transaction buffer changes # of pins, thus destabalizing the LinkedTo arrays

Change 3006828 on 2016/06/08 by [REDACTED_PERSON]

	Fix for non-editor builds

Change 3006857 on 2016/06/08 by [REDACTED_PERSON]

	Investigating shutdown ensure, traced back to a static UEdGraphPin

Change 3006907 on 2016/06/08 by [REDACTED_PERSON]

	Noneditor build fix

Change 3006929 on 2016/06/08 by [REDACTED_PERSON]

	Deferring DeprecatedPins destruction until after UBlueprint has had a chance to fix up its watched pins, this is a better fix for #jira UE-31779

Change 3007133 on 2016/06/09 by [REDACTED_PERSON]

	Fix for issue in the profiler asserting creating pins that don't have unique names.
	#Proj BlueprintProfiler

	- I believe this was recently introduced with the changes to UEdGraphPin's

Change 3007964 on 2016/06/09 by [REDACTED_PERSON]

	Fix for  PinHelpers::UnresolvedPins being left with stale entries by undo/redo

Change 3007996 on 2016/06/09 by [REDACTED_PERSON]

	Added 'empty' keyword to Array Clear Node.

Change 3008007 on 2016/06/09 by [REDACTED_PERSON]

	Added 'negate' keyword to boolean NOT node

Change 3008011 on 2016/06/09 by [REDACTED_PERSON]

	Added Vector2D * Vector2D multiplication node

Change 3008014 on 2016/06/09 by [REDACTED_PERSON]

	Fixed Cannot connect Make Array node output to MakeArray input with split pins

Change 3008243 on 2016/06/09 by [REDACTED_PERSON]

	Fix for creation of FWeakGraphPinPtr from a pin that had been destroyed, client logic is still a bit broken in the case of the ClassDefaults node, but we're back to 'safe'

Change 3008289 on 2016/06/09 by [REDACTED_PERSON]

	Editor transaction saves all state before applying undo/redo buffers when using 'bFlip' flow. This prevents messing with the object graph in the middle of saving state that will be restored later

Change 3008422 on 2016/06/09 by [REDACTED_PERSON]

	Correct usage of GIsTransacting, replaced with Ar.IsTransacting() to correctly handle the case where we serialize after transacting but during the transaction (for instance, recompile blueprint in post undo, which we do quite a bit it turns out)

Change 3009164 on 2016/06/10 by [REDACTED_PERSON]

	Making changes to default values in the structure editor will now make changes to the structure without rebuilding the default values panel.

Change 3009165 on 2016/06/10 by [REDACTED_PERSON]

	Fixed Structure Default value editor collapses after undoing an alteration of a default value

Change 3009181 on 2016/06/10 by [REDACTED_PERSON]

	Fixed issue where modifying a default value in a Widget Blueprint would cause the Details Panel to refresh

Change 3009313 on 2016/06/10 by [REDACTED_PERSON]

	Addressing issues with function return nodes in multiple ways:

	    - Preventing users from deleting return nodes for overriden/inherited functions.

	    - Also making sure that we create terminals for out params when the return node is disconnected (and pruned).

	    - Lastly, ensuring that new return nodes adhere to the function's signature (for cases, like where you copy/paste a return node from a different function).

Change 3009595 on 2016/06/10 by [REDACTED_PERSON]

	EdGraphPinReference using PinId to resolve itself again, may create issues resolving pins created in compile

Change 3009774 on 2016/06/10 by [REDACTED_PERSON]

	Fix for bad logic in RemovePin introduced in 3004329, just a bad reading of the logic, missed an early return

Change 3009988 on 2016/06/10 by [REDACTED_PERSON]

	Prefer to use existing pins (based on PinId) when undoing/redoing pin serialization

Change 3010050 on 2016/06/10 by [REDACTED_PERSON]

	Fixed missing call to ssuper class's PostEditUndo, fixed UBehaviorTreeGraph::PostEditUndo accessing Pins before they have been resolved

Change 3010071 on 2016/06/10 by [REDACTED_PERSON]

	Fix for pasting when owning node has whitespace in result of GetPathName

Change 3010244 on 2016/06/11 by [REDACTED_PERSON]

	Fix for trivial copy/paste error, causes crash when copying/pasting nodes with text default values, part of UE-31870

Change 3010630 on 2016/06/13 by [REDACTED_PERSON]

	No longer relying on path name for pin resolution, path is unstable across graphs

Change 3010647 on 2016/06/13 by [REDACTED_PERSON]

	PR #2496: Updated KismetMathLibrary comparison descriptions for FDateTime and FTimespan. (Contributed by [REDACTED_PERSON])

Change 3011175 on 2016/06/13 by [REDACTED_PERSON]

	Updates the Blueprint Profiler so that it can correctly map entry/exit from tunnels based on instance.
	#Proj Kismet, BlueprintProfiler

	- Ensured that the trace paths contain the macro instance exec nodes
	- Selectively update stats in the tunnel exit site nodes based on valid exit sites to prevent cyclic updates.
	- Updated the comments in map tunnel entry to spare peoples sanity when trying to understand what that function does.

Change 3011271 on 2016/06/13 by [REDACTED_PERSON]

	This adds support for inherited blueprint classes to the blueprint profiler.
	#Proj BlueprintProfiler

Change 3011556 on 2016/06/13 by [REDACTED_PERSON]

	Fixed Crash when breaking link to a split pin in MakeArray that is an array type

Change 3011624 on 2016/06/13 by [REDACTED_PERSON]

	Fix for missing entries in MessageLog's source pin identification map. Bob T had originally populated this correctly, but somehow i lost it while iterating.

Change 3011984 on 2016/06/13 by [REDACTED_PERSON]

	Sanitizing parentpin's subpins when destroying a pin

Change 3012894 on 2016/06/14 by [REDACTED_PERSON]

	[UE-30922] Ensure that customized defaults are propagated to new instances at construction time during non-Actor-based Blueprint class reinstancing.

	change summary:
	- modified FBlueprintCompileReinstancer::ReplaceInstancesOfClass_Inner() to use the reinstanced archetype object as the template object during construction of the new instance for non-Actor-based Blueprint class types.

Change 3013037 on 2016/06/14 by [REDACTED_PERSON]

	Fixed Crash when connecting to a split pin in a MakeArray node that has no connections

Change 3014846 on 2016/06/15 by [REDACTED_PERSON]

	No longer using FText::IsLetter to parse math expression nodes, that function is very slow. $x is now a valid math expression variable name (genereated a compile error prior to this change)

Change 3015014 on 2016/06/15 by [REDACTED_PERSON]

	Removing poorly implement IsLetter function

Change 3015142 on 2016/06/15 by [REDACTED_PERSON]

	More intentional about removing subpins, prevents stale iterator on split pin collapse

Change 3016326 on 2016/06/16 by [REDACTED_PERSON]

	Fixed MakeArray node does not reset to wildcard when breaking links with split struct pins that have default values

Change 3016494 on 2016/06/16 by [REDACTED_PERSON]

	Fixed Crash when dragging a component into the Event Graph that's inherited from a C++ class

Change 3016557 on 2016/06/16 by [REDACTED_PERSON]

	Explicit copy/move of string data for FText, removes some redundant copying and object construction/destruction [which could be optimzed away], saves 2-3 seconds in my 80s load asset benchmark

Change 3016577 on 2016/06/16 by [REDACTED_PERSON]

	Fixed compiler warning for hidden member variable in FBlueprintVarActionDetails::GetVariableReplicationType

Change 3016906 on 2016/06/16 by [REDACTED_PERSON]

	Back out changelist 3016557
	This will be done by [REDACTED_PERSON] in Dev-Editor

Change 3018081 on 2016/06/17 by [REDACTED_PERSON]

	[UE-31832] PR #2486: Expose UInheritableComponentHandler::GetAllTemplates() outside of editor (Contributed by [REDACTED_PERSON])

Change 3018402 on 2016/06/17 by [REDACTED_PERSON]

	Missing include

Change 3018426 on 2016/06/17 by [REDACTED_PERSON]

	Fixed MakeArray node with split pins and no connections does not paste correctly

Change 3018452 on 2016/06/17 by [REDACTED_PERSON]

	Moving the patching of instanced sub-objects out of CPFUO (where you can't rely on the target to be a replacement for the source) to FBlueprintEditorUtils::PatchCDOSubobjectsIntoExport(), and making it so PatchCDOSubobjectsIntoExport() is called regularly for Blueprint regeneration (on load).

Change 3018456 on 2016/06/17 by [REDACTED_PERSON]

	Fix for static analysis warning, this null check does nothing

Change 3018595 on 2016/06/17 by [REDACTED_PERSON]

	Fix for shadowed variable warning in CIS.

Change 3018699 on 2016/06/17 by [REDACTED_PERSON]

	Making MinimumAreaRectangle callable in Blueprints without world context (which is only needed for debug drawing).

Change 3019734 on 2016/06/20 by [REDACTED_PERSON]

	[UE-32064] Clone associated component template(s) when duplicating Blueprint function graphs containing one or more Add Component nodes.

	change summary:
	- added a UK2Node_AddComponent::PostDuplicate() override
	- moved UK2Node_AddComponent::PostPasteNode() logic into a helper method that's now called from both PostDuplicate() and PostPasteNode() overrides.

	notes:
	- will prevent getting into the scenario described in UE-31831

Change 3020635 on 2016/06/20 by [REDACTED_PERSON]

	Fix for bad cast in FCompilerResultsLog::Append, could cause crashes in clients of this function (math expressions nodes occasionally do when they fail to compile)

Change 3020894 on 2016/06/21 by [REDACTED_PERSON]

	#2522: Interface UProperties can ExposeOnSpawn (in Blueprints) (Contributed by [REDACTED_PERSON])

Change 3020958 on 2016/06/21 by [REDACTED_PERSON]

	This improves the way key events are detected in the blueprint profiler, preventing duplicate event entries when pressed and released are both wired. It also catches a bug with the compiler instrumentation flag when compiling.
	#Proj BlueprintProfiler, UnrealEd

Change 3021316 on 2016/06/21 by [REDACTED_PERSON]

	Fixed issue where Copy/Paste of event nodes would not retain link information

Change 3021826 on 2016/06/21 by [REDACTED_PERSON]

	[UE-31831] Fix up AddComponent nodes on load if they are not associated with a unique template object.

	change summary:
	- added external linkage to UK2Node_AddComponent::MakeNewComponentTemplate(), and switched it to be a public API
	- modified FBlueprintEditorUtils::UpdateComponentTemplates() (as this is already called on Blueprint load) to detect/warn and correct non-unique templates

Change 3022047 on 2016/06/21 by [REDACTED_PERSON]

	Fixed issue where copy/paste of return nodes would not preserve value or link data

Change 3022619 on 2016/06/22 by [REDACTED_PERSON]

	A static/persistent information (the mechanism is similar to AssetRegistrySearchable) about DynamicClass is added.
	It's necessary since DynamicClasses are not handled as regular assets by AssetRegistry.
	Fixed GameplayCueManager. Nativized cues can be found.

	This is an early version of the feature. Amount of stored persistent data can be extended (but it would increase memory-usage).

Change 3022654 on 2016/06/22 by [REDACTED_PERSON]

	FBackendHelperStaticSearchableValues -fixed  too strict ensure

Change 3023067 on 2016/06/22 by [REDACTED_PERSON]

	Config settings from super class are not applied (at runtime) to nativized Blueprints . So all "config" properties are filled in constructor.

Change 3023222 on 2016/06/22 by [REDACTED_PERSON]

	Fixed MakeArray node elements break when editing struct elements

Change 3023405 on 2016/06/22 by [REDACTED_PERSON]

	Making sure sub-objects get instanced for Blueprint CDOs that had their FObjectInitializer deferred (happens when the super CDO hasn't been fully serialized). By the time the deferred FObjectInitializer is ran, the sub-objects have been assigned a RF_NeedLoad flag (where they normally wouldn't have one right after construction, when the initialization is usually ran).

Change 3023992 on 2016/06/22 by [REDACTED_PERSON]

	Fixed an issue where hovering on/off a reroute node (toggling the comment bubble visibility) would create extraneous undo transactions.

diff --git a/Engine/Config/BaseEngine.ini b/Engine/Config/BaseEngine.ini
--- a/Engine/Config/BaseEngine.ini
+++ b/Engine/Config/BaseEngine.ini
@@ -480,6 +480,11 @@ bUseStreamingPause=false
 +K2FieldRedirects=(OldFieldName="KismetMathLibrary.BreakRot", NewFieldName="KismetMathLibrary.BreakRotator")
 +K2FieldRedirects=(OldFieldName="KismetMathLibrary.MakeRot", NewFieldName="KismetMathLibrary.MakeRotator")
 
+; Made UEdGraphPin not a UObject to reduce GC overhead, improve load times:
++TaggedPropertyRedirects=(ClassName="EdGraphNode",OldPropertyName="Pins",NewPropertyName="DeprecatedPins")
++TaggedPropertyRedirects=(ClassName="Blueprint",OldPropertyName="PinWatches",NewPropertyName="DeprecatedPinWatches")
++ActiveClassRedirects=(OldClassName="EdGraphPin",NewClassName="/Script/Engine.EdGraphPin_Deprecated")
+
 ;AIModule redirections
 +ActiveClassRedirects=(OldClassName="AIController",NewClassName="/Script/AIModule.AIController")
 +ActiveClassRedirects=(OldClassName="AIResourceInterface",NewClassName="/Script/AIModule.AIResourceInterface")
diff --git a/Engine/Source/Developer/BlueprintCompilerCppBackend/Private/BlueprintCompilerCppBackend.cpp b/Engine/Source/Developer/BlueprintCompilerCppBackend/Private/BlueprintCompilerCppBackend.cpp
--- a/Engine/Source/Developer/BlueprintCompilerCppBackend/Private/BlueprintCompilerCppBackend.cpp
+++ b/Engine/Source/Developer/BlueprintCompilerCppBackend/Private/BlueprintCompilerCppBackend.cpp
@@ -154,37 +154,46 @@ void FBlueprintCompilerCppBackend::EmitCastBetweenInterfacesStatement(FEmitterLo
 	EmitterContext.AddLine(FString::Printf(TEXT("}")));
 }
 
-void FBlueprintCompilerCppBackend::EmitCastInterfaceToObjStatement(FEmitterLocalContext& EmitterContext, FKismetFunctionContext& FunctionContext, FBlueprintCompiledStatement& Statement)
-{
-	FString ClassToCastTo = TermToText(EmitterContext, Statement.RHS[0]);
-	FString InputInterface = TermToText(EmitterContext, Statement.RHS[1]);
-	FString ResultObject = TermToText(EmitterContext, Statement.LHS, true, false);
-	FString InputObject = FString::Printf(TEXT("%s.GetObjectRef()"), *InputInterface);
-	EmitterContext.AddLine(FString::Printf(TEXT("%s = Cast<%s>(%s);"),*ResultObject, *ClassToCastTo, *InputObject));
-}
-
-void FBlueprintCompilerCppBackend::EmitDynamicCastStatement(FEmitterLocalContext& EmitterContext, FKismetFunctionContext& FunctionContext, FBlueprintCompiledStatement& Statement)
+static void EmitGenericCast(FEmitterLocalContext& EmitterContext, UClass* ClassPtr, const FString& ObjectValue, const FString& ResultObject)
 {
-	UClass* ClassPtr = CastChecked<UClass>(Statement.RHS[0]->ObjectLiteral);
 	check(ClassPtr != nullptr);
-	
-	const FString ObjectValue = TermToText(EmitterContext, Statement.RHS[1]);
-	const FString CastedValue = TermToText(EmitterContext, Statement.LHS, true, false);
 
 	auto BPGC = Cast<UBlueprintGeneratedClass>(ClassPtr);
 	if (BPGC && !EmitterContext.Dependencies.WillClassBeConverted(BPGC))
 	{
 		const FString NativeClass = FEmitHelper::GetCppName(EmitterContext.GetFirstNativeOrConvertedClass(ClassPtr));
 		const FString TargetClass = EmitterContext.FindGloballyMappedObject(ClassPtr, UClass::StaticClass(), true);
-		EmitterContext.AddLine(FString::Printf(TEXT("%s = NoNativeCast<%s>(%s, %s);"), *CastedValue, *NativeClass, *TargetClass, *ObjectValue));
+		EmitterContext.AddLine(FString::Printf(TEXT("%s = NoNativeCast<%s>(%s, %s);"), *ResultObject, *NativeClass, *TargetClass, *ObjectValue));
 	}
 	else
 	{
 		const FString TargetClass = FEmitHelper::GetCppName(ClassPtr);
-		EmitterContext.AddLine(FString::Printf(TEXT("%s = Cast<%s>(%s);"), *CastedValue, *TargetClass, *ObjectValue));
+		EmitterContext.AddLine(FString::Printf(TEXT("%s = Cast<%s>(%s);"), *ResultObject, *TargetClass, *ObjectValue));
 	}
 }
 
+void FBlueprintCompilerCppBackend::EmitCastInterfaceToObjStatement(FEmitterLocalContext& EmitterContext, FKismetFunctionContext& FunctionContext, FBlueprintCompiledStatement& Statement)
+{
+	const FString InputInterface = TermToText(EmitterContext, Statement.RHS[1]);
+	const FString ResultObject = TermToText(EmitterContext, Statement.LHS, true, false);
+
+	EmitGenericCast(EmitterContext
+		, CastChecked<UClass>(Statement.RHS[0]->ObjectLiteral)
+		, FString::Printf(TEXT("%s.GetObjectRef()"), *InputInterface)
+		, ResultObject);
+}
+
+void FBlueprintCompilerCppBackend::EmitDynamicCastStatement(FEmitterLocalContext& EmitterContext, FKismetFunctionContext& FunctionContext, FBlueprintCompiledStatement& Statement)
+{
+	const FString ObjectValue = TermToText(EmitterContext, Statement.RHS[1]);
+	const FString ResultObject = TermToText(EmitterContext, Statement.LHS, true, false);
+
+	EmitGenericCast(EmitterContext
+		, CastChecked<UClass>(Statement.RHS[0]->ObjectLiteral)
+		, ObjectValue
+		, ResultObject);
+}
+
 void FBlueprintCompilerCppBackend::EmitMetaCastStatement(FEmitterLocalContext& EmitterContext, FKismetFunctionContext& FunctionContext, FBlueprintCompiledStatement& Statement)
 {
 	FString DesiredClass = TermToText(EmitterContext, Statement.RHS[0]);
@@ -404,11 +413,24 @@ FString FBlueprintCompilerCppBackend::EmitSwitchValueStatmentInner(FEmitterLocal
 
 	for (int32 TermIndex = TermsBeforeCases; TermIndex < (NumCases * TermsPerCase); TermIndex += TermsPerCase)
 	{
+		auto TermToRef = [&](const FBPTerminal* Term) -> FString
+		{
+			const FString TermEvaluation = TermToText(EmitterContext, Term); //should bGetter be false ?
+			if (Term->bIsLiteral) //TODO it should be done for every term, that cannot be handled as reference.
+			{
+				const FString LocalVarName = EmitterContext.GenerateUniqueLocalName();
+				EmitterContext.AddLine(FString::Printf(TEXT("%s %s = %s;"), *ValueDeclaration, *LocalVarName, *TermEvaluation));
+				return LocalVarName;
+			}
+			return TermEvaluation;
+		};
+		const FString Term0_Index = TermToText(EmitterContext, Statement.RHS[TermIndex]);
+		const FString Term1_Ref = TermToRef(Statement.RHS[TermIndex + 1]);
 		Result += FString::Printf(TEXT(", TSwitchPair<%s, %s>(%s, %s)")
 			, *IndexDeclaration
 			, *ValueDeclaration
-			, *TermToText(EmitterContext, Statement.RHS[TermIndex])
-			, *TermToText(EmitterContext, Statement.RHS[TermIndex + 1]));
+			, *Term0_Index
+			, *Term1_Ref);
 	}
 
 	Result += TEXT(")");
@@ -482,6 +504,54 @@ FString FBlueprintCompilerCppBackend::EmitMethodInputParameterList(FEmitterLocal
 	return Result;
 }
 
+static FString CustomThunkFunctionPostfix(FBlueprintCompiledStatement& Statement)
+{
+	/* Some native structures have no operator==. There are special versions of array functions for them (see GeneratedCodeHelpers.h). */
+
+	check(Statement.FunctionToCall);
+
+	int32 NumParams = 0;
+	FBPTerminal* ArrayTerm = nullptr;
+	for (TFieldIterator<UProperty> PropIt(Statement.FunctionToCall); PropIt && (PropIt->PropertyFlags & CPF_Parm); ++PropIt)
+	{
+		UProperty* FuncParamProperty = *PropIt;
+		if (!FuncParamProperty->HasAnyPropertyFlags(CPF_ReturnParm))
+		{
+			if (UArrayProperty* ArrayProperty = Cast<UArrayProperty>(*PropIt))
+			{
+				ArrayTerm = Statement.RHS[NumParams];
+				ensure(ArrayTerm && ArrayTerm->Type.bIsArray);
+				break;
+			}
+			NumParams++;
+		}
+	}
+
+	FString FunctionName = Statement.FunctionToCall->GetName();
+	if (ArrayTerm 
+		&& (FunctionName == TEXT("Array_Find") 
+			|| FunctionName == TEXT("Array_Contains") 
+			|| FunctionName == TEXT("Array_RemoveItem")
+			|| FunctionName == TEXT("Array_AddUnique")
+			))
+	{
+		if (UEdGraphSchema_K2::PC_Text == ArrayTerm->Type.PinCategory)
+		{
+			return FString(TEXT("_FText"));
+		}
+
+		if (UEdGraphSchema_K2::PC_Struct == ArrayTerm->Type.PinCategory)
+		{
+			UScriptStruct* Struct = Cast<UScriptStruct>(ArrayTerm->Type.PinSubCategoryObject.Get());
+			if (Struct && Struct->IsNative() && (0 == (Struct->StructFlags & STRUCT_NoExport)))
+			{
+				return FString(TEXT("_Struct"));
+			}
+		}
+	}
+	return FString();
+}
+
 FString FBlueprintCompilerCppBackend::EmitCallStatmentInner(FEmitterLocalContext& EmitterContext, FBlueprintCompiledStatement& Statement, bool bInline, FString PostFix)
 {
 	const bool bCallOnDifferentObject = Statement.FunctionContext && (Statement.FunctionContext->Name != TEXT("self"));
@@ -489,12 +559,14 @@ FString FBlueprintCompilerCppBackend::EmitCallStatmentInner(FEmitterLocalContext
 	const bool bUseSafeContext = bCallOnDifferentObject && !bStaticCall;
 	const bool bAnyInterfaceCall = bCallOnDifferentObject && Statement.FunctionContext && (UEdGraphSchema_K2::PC_Interface == Statement.FunctionContext->Type.PinCategory);
 	const bool bInterfaceCallExecute = bAnyInterfaceCall && Statement.FunctionToCall && Statement.FunctionToCall->HasAnyFunctionFlags(FUNC_Event | FUNC_BlueprintEvent);
+	const bool bNativeEvent = FEmitHelper::ShouldHandleAsNativeEvent(Statement.FunctionToCall, false);
 
 	const UClass* CurrentClass = EmitterContext.GetCurrentlyGeneratedClass();
 	const UClass* SuperClass = CurrentClass ? CurrentClass->GetSuperClass() : nullptr;
 	const UClass* OriginalSuperClass = SuperClass ? EmitterContext.Dependencies.FindOriginalClass(SuperClass) : nullptr;
 	const UFunction* ActualParentFunction = (Statement.bIsParentContext && OriginalSuperClass) ? OriginalSuperClass->FindFunctionByName(Statement.FunctionToCall->GetFName(), EIncludeSuperFlag::IncludeSuper) : nullptr;
-	const FString FunctionToCallOriginalName = FEmitHelper::GetCppName(ActualParentFunction ? ActualParentFunction : FEmitHelper::GetOriginalFunction(Statement.FunctionToCall)) + PostFix;
+	// if(Statement.bIsParentContext && bNativeEvent) then name is constructed from original function with "_Implementation postfix
+	const FString FunctionToCallOriginalName = FEmitHelper::GetCppName((ActualParentFunction && !bNativeEvent) ? ActualParentFunction : FEmitHelper::GetOriginalFunction(Statement.FunctionToCall)) + PostFix;
 	const bool bIsFunctionValidToCallFromBP = !ActualParentFunction || ActualParentFunction->HasAnyFunctionFlags(FUNC_Native) || (ActualParentFunction->Script.Num() > 0);
 
 	if (!bIsFunctionValidToCallFromBP)
@@ -535,7 +607,7 @@ FString FBlueprintCompilerCppBackend::EmitCallStatmentInner(FEmitterLocalContext
 			{
 				const FString Value = TermToText(EmitterContext, ValueTerminal);
 				const FString EnumName = FEmitHelper::GetCppName(Enum);
-				Result += FString::Printf(TEXT("%s__GetUserFriendlyName(EnumToByte<%s>(%s))"), *EnumName, *EnumName, *Value);
+				Result += FString::Printf(TEXT("%s__GetUserFriendlyName(EnumToByte<%s>(TEnumAsByte<%s>(%s)))"), *EnumName, *EnumName, *EnumName, *Value);
 				return true;
 			}
 		}
@@ -560,6 +632,9 @@ FString FBlueprintCompilerCppBackend::EmitCallStatmentInner(FEmitterLocalContext
 			auto FunctionOwner = Statement.FunctionToCall->GetOwnerClass();
 			auto OwnerBPGC = Cast<UBlueprintGeneratedClass>(FunctionOwner);
 			const bool bUnconvertedClass = OwnerBPGC && !EmitterContext.Dependencies.WillClassBeConverted(OwnerBPGC);
+			const bool bIsCustomThunk = bStaticCall && Statement.FunctionToCall->GetBoolMetaData(TEXT("CustomThunk"))
+				|| Statement.FunctionToCall->HasMetaData(TEXT("CustomStructureParam"))
+				|| Statement.FunctionToCall->HasMetaData(TEXT("ArrayParm"));
 			if (bUnconvertedClass)
 			{
 				ensure(!Statement.bIsParentContext); //unsupported yet
@@ -571,9 +646,6 @@ FString FBlueprintCompilerCppBackend::EmitCallStatmentInner(FEmitterLocalContext
 			}
 			else if (bStaticCall)
 			{
-				const bool bIsCustomThunk = Statement.FunctionToCall->GetBoolMetaData(TEXT("CustomThunk"))
-					|| Statement.FunctionToCall->HasMetaData(TEXT("CustomStructureParam"))
-					|| Statement.FunctionToCall->HasMetaData(TEXT("ArrayParm"));
 				auto OwnerClass = Statement.FunctionToCall->GetOuterUClass();
 				Result += bIsCustomThunk ? TEXT("FCustomThunkTemplates::") : FString::Printf(TEXT("%s::"), *FEmitHelper::GetCppName(OwnerClass));
 			}
@@ -588,7 +660,12 @@ FString FBlueprintCompilerCppBackend::EmitCallStatmentInner(FEmitterLocalContext
 			}
 			Result += FunctionToCallOriginalName;
 
-			if (Statement.bIsParentContext && FEmitHelper::ShouldHandleAsNativeEvent(Statement.FunctionToCall, false))
+			if (bIsCustomThunk)
+			{
+				Result += CustomThunkFunctionPostfix(Statement);
+			}
+
+			if (Statement.bIsParentContext && bNativeEvent)
 			{
 				ensure(!bCallOnDifferentObject);
 				Result += TEXT("_Implementation");
@@ -735,7 +812,16 @@ FString FBlueprintCompilerCppBackend::TermToText(FEmitterLocalContext& EmitterCo
 				{
 					ResultPath = ContextStr + TEXT("->");
 				}
+
 				ResultPath += FEmitHelper::GetCppName(Term->AssociatedVarProperty);
+
+				// convert bitfield to bool...
+				UBoolProperty* BoolProperty = Cast<UBoolProperty>(Term->AssociatedVarProperty);
+				if (bGetter && BoolProperty && !BoolProperty->IsNativeBool())
+				{
+					//TODO: the result still cannot be used as reference
+					ResultPath = FString::Printf(TEXT("(%s != 0)"), *ResultPath);
+				}
 			}
 		}
 		else
@@ -781,7 +867,7 @@ FString FBlueprintCompilerCppBackend::TermToText(FEmitterLocalContext& EmitterCo
 			EmitterContext.AddLine(*FString::Printf(TEXT("%s %s = %s;"), *CppType, *DefaultValueVariable, *DefaultValueConstructor));
 			return FString::Printf(TEXT("((%s) ? (%s) : (%s))"), *Conditions, *ResultPath, *DefaultValueVariable);
 		}
-		return ResultPath;
+		return ResultPath; 
 	}
 }
 
diff --git a/Engine/Source/Developer/BlueprintCompilerCppBackend/Private/BlueprintCompilerCppBackendBase.cpp b/Engine/Source/Developer/BlueprintCompilerCppBackend/Private/BlueprintCompilerCppBackendBase.cpp
--- a/Engine/Source/Developer/BlueprintCompilerCppBackend/Private/BlueprintCompilerCppBackendBase.cpp
+++ b/Engine/Source/Developer/BlueprintCompilerCppBackend/Private/BlueprintCompilerCppBackendBase.cpp
@@ -222,6 +222,8 @@ FString FBlueprintCompilerCppBackendBase::GenerateCodeFromClass(UClass* SourceCl
 	UClass* OriginalSourceClass = Dependencies.FindOriginalClass(SourceClass);
 	ensure(OriginalSourceClass != SourceClass);
 
+	const bool bHasStaticSearchableValues = FBackendHelperStaticSearchableValues::HasSearchableValues(SourceClass);
+
 	// Class declaration
 	const bool bIsInterface = SourceClass->IsChildOf<UInterface>();
 	if (bIsInterface)
@@ -244,9 +246,15 @@ FString FBlueprintCompilerCppBackendBase::GenerateCodeFromClass(UClass* SourceCl
 			AdditionalMD.Add(ReplaceConvertedMD);
 		}
 
-		// AdditionalMD.Add(FString::Printf(TEXT("CustomDynamicClassInitialization=\"%s::__CustomDynamicClassInitialization\""), *CppClassName));
+		if (bHasStaticSearchableValues)
+		{
+			AdditionalMD.Add(FBackendHelperStaticSearchableValues::GenerateClassMetaData(SourceClass));
+		}
 
-		EmitterContext.Header.AddLine(FString::Printf(TEXT("UCLASS(%s%s)")
+		// AdditionalMD.Add(FString::Printf(TEXT("CustomDynamicClassInitialization=\"%s::__CustomDynamicClassInitialization\""), *CppClassName));
+		const FString DefinedConfigName = (OriginalSourceClass->ClassConfigName == NAME_None) ? FString() : FString::Printf(TEXT("config=%s, "), *OriginalSourceClass->ClassConfigName.ToString());
+		EmitterContext.Header.AddLine(FString::Printf(TEXT("UCLASS(%s%s%s)")
+			, *DefinedConfigName
 			, (!SourceClass->IsChildOf<UBlueprintFunctionLibrary>()) ? TEXT("Blueprintable, BlueprintType, ") : TEXT("")
 			, *FEmitHelper::HandleMetaData(nullptr, false, &AdditionalMD)));
 
@@ -280,7 +288,12 @@ FString FBlueprintCompilerCppBackendBase::GenerateCodeFromClass(UClass* SourceCl
 		EmitterContext.Header.AddLine(TEXT("virtual void PostLoadSubobjects(FObjectInstancingGraph* OuterInstanceGraph) override;"));
 		EmitterContext.Header.AddLine(TEXT("static void __StaticDependenciesAssets(TArray<FBlueprintDependencyData>& AssetsToLoad);"));
 		EmitterContext.Header.AddLine(TEXT("static void __CustomDynamicClassInitialization(UDynamicClass* InDynamicClass);"));
-
+		EmitterContext.Header.AddLine(TEXT("static void __CustomDynamicClassInitializationUsedAssets(UDynamicClass* InDynamicClass);"));
+		if (bHasStaticSearchableValues)
+		{
+			FBackendHelperStaticSearchableValues::EmitFunctionDeclaration(EmitterContext);
+			FBackendHelperStaticSearchableValues::EmitFunctionDefinition(EmitterContext);
+		}
 		FEmitDefaultValueHelper::GenerateConstructor(EmitterContext);
 	}
 
@@ -308,6 +321,13 @@ FString FBlueprintCompilerCppBackendBase::GenerateCodeFromClass(UClass* SourceCl
 	EmitterContext.Header.DecreaseIndent();
 	EmitterContext.Header.AddLine(TEXT("};"));
 
+	if (!bIsInterface)
+	{
+		// must be called after GenerateConstructor
+		// now we knows which assets are directly used in source code
+		FEmitDefaultValueHelper::GenerateCustomDynamicClassInitializationUsedAssets(EmitterContext);
+	}
+
 	FEmitHelper::EmitLifetimeReplicatedPropsImpl(EmitterContext);
 
 	EmitterContext.Body.AddLine(TEXT("#ifdef _MSC_VER"));
@@ -697,12 +717,13 @@ FString FBlueprintCompilerCppBackendBase::GenerateCodeFromEnum(UUserDefinedEnum*
 
 	for (int32 Index = 0; Index < SourceEnum->NumEnums(); ++Index)
 	{
-		const FString ElemName = EnumItemName(Index);
+		const FString ElemCppName = EnumItemName(Index);
 		const int32 ElemValue = SourceEnum->GetValueByIndex(Index);
 
 		const FString& DisplayNameMD = SourceEnum->GetMetaData(TEXT("DisplayName"), ElemValue);// TODO: value or index?
-		const FString Meta = DisplayNameMD.IsEmpty() ? FString() : FString::Printf(TEXT("UMETA(DisplayName = \"%s\")"), *DisplayNameMD.ReplaceCharWithEscapedChar());
-		Header.AddLine(FString::Printf(TEXT("%s = %d %s,"), *ElemName, ElemValue, *Meta));
+		const FString MetaDisplayName = DisplayNameMD.IsEmpty() ? FString() : FString::Printf(TEXT("DisplayName = \"%s\","), *DisplayNameMD.ReplaceCharWithEscapedChar());
+		const FString MetaOverrideName = FString::Printf(TEXT("OverrideName = \"%s\""), *SourceEnum->GetNameByIndex(Index).ToString());
+		Header.AddLine(FString::Printf(TEXT("%s = %d UMETA(%s%s),"), *ElemCppName, ElemValue, *MetaDisplayName, *MetaOverrideName));
 	}
 
 	Header.DecreaseIndent();
@@ -744,8 +765,9 @@ FString FBlueprintCompilerCppBackendBase::GenerateCodeFromStruct(UUserDefinedStr
 	// use GetBaseFilename() so that we can coordinate #includes and filenames
 	EmitFileBeginning(FEmitHelper::GetBaseFilename(SourceStruct), EmitterContext, true, true);
 
+	const FString CppStructName = FEmitHelper::GetCppName(SourceStruct);
 	EmitterContext.Header.AddLine(FString::Printf(TEXT("USTRUCT(BlueprintType, %s)"), *FEmitHelper::ReplaceConvertedMetaData(SourceStruct)));
-	EmitterContext.Header.AddLine(FString::Printf(TEXT("struct %s"), *FEmitHelper::GetCppName(SourceStruct)));
+	EmitterContext.Header.AddLine(FString::Printf(TEXT("struct %s"), *CppStructName));
 	EmitterContext.Header.AddLine(TEXT("{"));
 	EmitterContext.Header.AddLine(TEXT("public:"));
 	EmitterContext.Header.IncreaseIndent();
@@ -753,6 +775,14 @@ FString FBlueprintCompilerCppBackendBase::GenerateCodeFromStruct(UUserDefinedStr
 	EmitStructProperties(EmitterContext, SourceStruct);
 
 	FEmitDefaultValueHelper::GenerateGetDefaultValue(SourceStruct, EmitterContext);
+
+	EmitterContext.Header.AddLine(FString::Printf(TEXT("bool operator== (const %s& __Other) const"), *CppStructName));
+	EmitterContext.Header.AddLine(TEXT("{"));
+	EmitterContext.Header.IncreaseIndent();
+	EmitterContext.Header.AddLine(FString::Printf(TEXT("return %s::StaticStruct()->CompareScriptStruct(this, &__Other, 0);"), *CppStructName));
+	EmitterContext.Header.DecreaseIndent();
+	EmitterContext.Header.AddLine(TEXT("};"));
+
 	EmitterContext.Header.DecreaseIndent();
 	EmitterContext.Header.AddLine(TEXT("};"));
 
diff --git a/Engine/Source/Developer/BlueprintCompilerCppBackend/Private/BlueprintCompilerCppBackendGatherDependencies.cpp b/Engine/Source/Developer/BlueprintCompilerCppBackend/Private/BlueprintCompilerCppBackendGatherDependencies.cpp
--- a/Engine/Source/Developer/BlueprintCompilerCppBackend/Private/BlueprintCompilerCppBackendGatherDependencies.cpp
+++ b/Engine/Source/Developer/BlueprintCompilerCppBackend/Private/BlueprintCompilerCppBackendGatherDependencies.cpp
@@ -337,8 +337,12 @@ void FGatherConvertedClassDependencies::DependenciesForHeader()
 
 	for (auto Obj : ObjectsToCheck)
 	{
-		auto Property = Cast<const UProperty>(Obj);
-		auto OwnerProperty = IsValid(Property) ? Property->GetOwnerProperty() : nullptr;
+		const UProperty* Property = Cast<const UProperty>(Obj);
+		if (const UArrayProperty* ArrayProperty = Cast<UArrayProperty>(Property))
+		{
+			Property = ArrayProperty->Inner;
+		}
+		const UProperty* OwnerProperty = IsValid(Property) ? Property->GetOwnerProperty() : nullptr;
 		const bool bIsParam = OwnerProperty && (0 != (OwnerProperty->PropertyFlags & CPF_Parm)) && OwnerProperty->IsIn(OriginalStruct);
 		const bool bIsMemberVariable = OwnerProperty && (OwnerProperty->GetOuter() == OriginalStruct);
 		if (bIsParam || bIsMemberVariable)
@@ -366,16 +370,21 @@ void FGatherConvertedClassDependencies::DependenciesForHeader()
 			{
 				IncludeInHeader.Add(InterfaceProperty->InterfaceClass);
 			}
-			/* Delegate signatures are recreated in local scope anyway.
 			else if (auto DelegateProperty = Cast<const UDelegateProperty>(Property))
 			{
-				IncludeInHeader.Add(DelegateProperty->SignatureFunction ? DelegateProperty->SignatureFunction->GetOwnerClass() : nullptr);
+				IncludeInHeader.Add(DelegateProperty->SignatureFunction ? DelegateProperty->SignatureFunction->GetOwnerStruct() : nullptr);
 			}
+			/* MC Delegate signatures are recreated in local scope anyway.
 			else if (auto MulticastDelegateProperty = Cast<const UMulticastDelegateProperty>(Property))
 			{
 				IncludeInHeader.Add(MulticastDelegateProperty->SignatureFunction ? MulticastDelegateProperty->SignatureFunction->GetOwnerClass() : nullptr);
 			}
 			*/
+			else if (const UByteProperty* ByteProperty = Cast<const UByteProperty>(Property))
+			{ 
+				// HeaderReferenceFinder.FindReferences(Obj); cannot find this enum..
+				IncludeInHeader.Add(ByteProperty->Enum);
+			}
 			else
 			{
 				HeaderReferenceFinder.FindReferences(Obj);
diff --git a/Engine/Source/Developer/BlueprintCompilerCppBackend/Private/BlueprintCompilerCppBackendUtils.cpp b/Engine/Source/Developer/BlueprintCompilerCppBackend/Private/BlueprintCompilerCppBackendUtils.cpp
--- a/Engine/Source/Developer/BlueprintCompilerCppBackend/Private/BlueprintCompilerCppBackendUtils.cpp
+++ b/Engine/Source/Developer/BlueprintCompilerCppBackend/Private/BlueprintCompilerCppBackendUtils.cpp
@@ -142,13 +142,30 @@ FString FEmitterLocalContext::FindGloballyMappedObject(const UObject* Object, co
 		return FString::Printf(TEXT("%s::StaticStruct()"), *FEmitHelper::GetCppName(UDS));
 	}
 
+	if (const UUserDefinedEnum* UDE = Cast<UUserDefinedEnum>(Object))
+	{
+		const int32 EnumIndex = EnumsInCurrentClass.IndexOfByKey(UDE);
+		if (INDEX_NONE != EnumIndex)
+		{
+			return FString::Printf(TEXT("CastChecked<%s>(CastChecked<UDynamicClass>(%s::StaticClass())->ReferencedConvertedFields[%d])")
+				, *ClassString()
+				, *FEmitHelper::GetCppName(ActualClass)
+				, EnumIndex);
+		}
+	}
+
 	// TODO: handle subobjects
 	ensure(!bLoadIfNotFound || Object);
 	if (Object && (bLoadIfNotFound || bTryUsedAssetsList))
 	{
 		if (bTryUsedAssetsList)
 		{
-			const int32 AssetIndex = Dependencies.Assets.IndexOfByKey(Object);
+			int32 AssetIndex = UsedObjectInCurrentClass.IndexOfByKey(Object);
+			if (INDEX_NONE == AssetIndex && Dependencies.Assets.Contains(Object))
+			{
+				AssetIndex = UsedObjectInCurrentClass.Add(Object);
+			}
+
 			if (INDEX_NONE != AssetIndex)
 			{
 				return FString::Printf(TEXT("CastChecked<%s>(CastChecked<UDynamicClass>(%s::StaticClass())->UsedAssets[%d], ECastCheckedType::NullAllowed)")
@@ -171,20 +188,24 @@ FString FEmitterLocalContext::FindGloballyMappedObject(const UObject* Object, co
 
 FString FEmitterLocalContext::ExportTextItem(const UProperty* Property, const void* PropertyValue) const
 {
+	const uint32 LocalExportCPPFlags = EPropertyExportCPPFlags::CPPF_CustomTypeName
+		| EPropertyExportCPPFlags::CPPF_NoConst
+		| EPropertyExportCPPFlags::CPPF_NoRef
+		| EPropertyExportCPPFlags::CPPF_NoStaticArray
+		| EPropertyExportCPPFlags::CPPF_BlueprintCppBackend;
 	if (auto ArrayProperty = Cast<const UArrayProperty>(Property))
 	{
-		const uint32 LocalExportCPPFlags = EPropertyExportCPPFlags::CPPF_CustomTypeName
-			| EPropertyExportCPPFlags::CPPF_NoConst
-			| EPropertyExportCPPFlags::CPPF_NoRef
-			| EPropertyExportCPPFlags::CPPF_NoStaticArray
-			| EPropertyExportCPPFlags::CPPF_BlueprintCppBackend;
 		const FString ConstPrefix = Property->HasMetaData(TEXT("NativeConstTemplateArg")) ? TEXT("const ") : TEXT("");
 		const FString TypeText = ExportCppDeclaration(ArrayProperty, EExportedDeclaration::Parameter, LocalExportCPPFlags, true, FString(), ConstPrefix);
 		return FString::Printf(TEXT("%s()"), *TypeText);
 	}
-
 	FString ValueStr;
 	Property->ExportTextItem(ValueStr, PropertyValue, PropertyValue, nullptr, EPropertyPortFlags::PPF_ExportCpp);
+	if (Property->IsA<UAssetObjectProperty>())
+	{
+		const FString TypeText = ExportCppDeclaration(Property, EExportedDeclaration::Parameter, LocalExportCPPFlags, true);
+		return FString::Printf(TEXT("%s(%s)"), *TypeText, *ValueStr);
+	}
 	return ValueStr;
 }
 
@@ -1244,8 +1265,8 @@ bool FEmitHelper::GenerateAutomaticCast(FEmitterLocalContext& EmitterContext, co
 		{
 			ensure(!RTypeEnum->IsA<UUserDefinedEnum>() || RTypeEnum->CppType.IsEmpty());
 			const FString EnumCppType = !RTypeEnum->CppType.IsEmpty() ? RTypeEnum->CppType : FEmitHelper::GetCppName(RTypeEnum);
-			OutCastBegin = FString::Printf(TEXT("EnumToByte<%s>("), *EnumCppType); 
-			OutCastEnd = TEXT(")");
+			OutCastBegin = FString::Printf(TEXT("EnumToByte<%s>(TEnumAsByte<%s>("), *EnumCppType, *EnumCppType);
+			OutCastEnd = TEXT("))");
 			return true;
 		}
 	}
@@ -1451,7 +1472,11 @@ FString FEmitHelper::AccessInaccessiblePropertyUsingOffset(FEmitterLocalContext&
 {
 	check(Property);
 
-	UE_LOG(LogK2Compiler, Warning, TEXT("AccessInaccessiblePropertyUsingOffset - NOEXPORT structure should be handled in a custom way: %s"), *GetPathNameSafe(Property->GetOwnerStruct()));
+	if (!FEmitDefaultValueHelper::SpecialStructureConstructor(Property->GetOwnerStruct(), nullptr, nullptr))
+	{
+		// no need to log warning for known structures
+		UE_LOG(LogK2Compiler, Warning, TEXT("AccessInaccessiblePropertyUsingOffset - NOEXPORT structure should be handled in a custom way: %s"), *GetPathNameSafe(Property->GetOwnerStruct()));
+	}
 
 	const int32 PropertyOffsetWithoutEditorOnlyMembers = HelperWithoutEditorOnlyMembers::OffsetWithoutEditorOnlyMembers(Property);
 	const int32 PropertySizeWithoutEditorOnlyMembers = HelperWithoutEditorOnlyMembers::SizeWithoutEditorOnlyMembers(Property);
@@ -1477,3 +1502,76 @@ FString FEmitHelper::AccessInaccessiblePropertyUsingOffset(FEmitterLocalContext&
 		, ElementSizeWithoutEditorOnlyMembers
 		, StaticArrayIdx);
 }
+
+bool FBackendHelperStaticSearchableValues::HasSearchableValues(UClass* Class)
+{
+	TArray<FStringClassReference> ClassesWithSearchableValues;
+	//TODO: export to .ini
+	ClassesWithSearchableValues.Add(FStringClassReference(TEXT("/Script/GameplayAbilities.GameplayCueNotify_Actor")));
+	ClassesWithSearchableValues.Add(FStringClassReference(TEXT("/Script/GameplayAbilities.GameplayCueNotify_Static")));
+
+	for (FStringClassReference& IterClassRef : ClassesWithSearchableValues)
+	{
+		UClass* IterClass = IterClassRef.ResolveClass();
+		if (IterClass && Class && Class->IsChildOf(IterClass))
+		{
+			return true;
+		}
+	}
+	return false;
+}
+
+FString FBackendHelperStaticSearchableValues::GetFunctionName()
+{
+	return TEXT("__InitializeStaticSearchableValues");
+}
+
+FString FBackendHelperStaticSearchableValues::GenerateClassMetaData(UClass* Class)
+{
+	const FString MetaDataName = TEXT("InitializeStaticSearchableValues");
+	const FString FunctionName = GetFunctionName();
+	return FString::Printf(TEXT("%s=\"%s\""), *MetaDataName, *FunctionName);
+}
+
+void FBackendHelperStaticSearchableValues::EmitFunctionDeclaration(FEmitterLocalContext& Context)
+{
+	const FString FunctionName = GetFunctionName();
+	Context.Header.AddLine(FString::Printf(TEXT("static void %s(TMap<FName, FName>& SearchableValues);"), *FunctionName));
+}
+
+void FBackendHelperStaticSearchableValues::EmitFunctionDefinition(FEmitterLocalContext& Context)
+{
+	auto BPGC = CastChecked<UBlueprintGeneratedClass>(Context.GetCurrentlyGeneratedClass());
+	const FString CppClassName = FEmitHelper::GetCppName(BPGC);
+	const FString FunctionName = GetFunctionName();
+
+	Context.Body.AddLine(FString::Printf(TEXT("void %s::%s(TMap<FName, FName>& SearchableValues)"), *CppClassName, *FunctionName));
+	Context.Body.AddLine(TEXT("{"));
+	Context.IncreaseIndent();
+
+	UClass* OriginalSourceClass = Context.Dependencies.FindOriginalClass(BPGC);
+	if(ensure(OriginalSourceClass))
+	{
+		FAssetData ClassAsset(OriginalSourceClass);
+
+		TArray<FName> TagPropertyNames;
+		//TODO: export to ini
+		TagPropertyNames.Add(FName(TEXT("GameplayCueName")));
+
+		for (const FName TagPropertyName : TagPropertyNames)
+		{
+			const FName FoundValue = ClassAsset.GetTagValueRef<FName>(TagPropertyName);
+			if (!FoundValue.IsNone())
+			{
+				Context.Body.AddLine(FString::Printf(TEXT("SearchableValues.Add(FName(TEXT(\"%s\")), FName(TEXT(\"%s\")));"), *TagPropertyName.ToString(), *FoundValue.ToString()));
+			}
+			else
+			{
+				UE_LOG(LogK2Compiler, Warning, TEXT("FBackendHelperStaticSearchableValues - None value. Tag: %s Asset: %s"), *TagPropertyName.ToString(), *GetPathNameSafe(OriginalSourceClass));
+			}
+		}
+	}
+
+	Context.Body.DecreaseIndent();
+	Context.Body.AddLine(TEXT("}"));
+}
\ No newline at end of file
diff --git a/Engine/Source/Developer/BlueprintCompilerCppBackend/Private/BlueprintCompilerCppBackendUtils.h b/Engine/Source/Developer/BlueprintCompilerCppBackend/Private/BlueprintCompilerCppBackendUtils.h
--- a/Engine/Source/Developer/BlueprintCompilerCppBackend/Private/BlueprintCompilerCppBackendUtils.h
+++ b/Engine/Source/Developer/BlueprintCompilerCppBackend/Private/BlueprintCompilerCppBackendUtils.h
@@ -43,6 +43,8 @@ struct FEmitterLocalContext
 
 	EGeneratedCodeType CurrentCodeType;
 
+	TArray<const UObject*> UsedObjectInCurrentClass;
+	TArray<const UUserDefinedEnum*> EnumsInCurrentClass;
 private:
 	int32 LocalNameIndexMax;
 
@@ -274,7 +276,9 @@ struct FEmitDefaultValueHelper
 
 	static void GenerateConstructor(FEmitterLocalContext& Context);
 
-	static void GenerateCustomDynamicClassInitialization(FEmitterLocalContext& Context);
+	static void GenerateCustomDynamicClassInitialization(FEmitterLocalContext& Context, TSharedPtr<FGatherConvertedClassDependencies> ParentDependencies);
+
+	static void GenerateCustomDynamicClassInitializationUsedAssets(FEmitterLocalContext& Context);
 
 	enum class EPropertyAccessOperator
 	{
@@ -295,7 +299,7 @@ struct FEmitDefaultValueHelper
 	static FString HandleClassSubobject(FEmitterLocalContext& Context, UObject* Object, FEmitterLocalContext::EClassSubobjectList ListOfSubobjectsTyp, bool bCreate, bool bInitilize);
 
 	// returns true, and fill OutResult, when the structure is handled in a custom way.
-	static bool SpecialStructureConstructor(const UScriptStruct* Struct, const uint8* ValuePtr, FString* OutResult);
+	static bool SpecialStructureConstructor(const UStruct* Struct, const uint8* ValuePtr, FString* OutResult);
 
 private:
 	// Returns native term, 
@@ -325,4 +329,14 @@ struct FBackendHelperAnim
 	static void AddHeaders(FEmitterLocalContext& EmitterContext);
 
 	static void CreateAnimClassData(FEmitterLocalContext& Context);
+};
+
+/** this struct helps generate a static function that initializes Static Searchable Values. */
+struct FBackendHelperStaticSearchableValues
+{
+	static bool HasSearchableValues(UClass* Class);
+	static FString GetFunctionName();
+	static FString GenerateClassMetaData(UClass* Class);
+	static void EmitFunctionDeclaration(FEmitterLocalContext& Context);
+	static void EmitFunctionDefinition(FEmitterLocalContext& Context);
 };
\ No newline at end of file
diff --git a/Engine/Source/Developer/BlueprintCompilerCppBackend/Private/BlueprintCompilerCppBackendValueHelper.cpp b/Engine/Source/Developer/BlueprintCompilerCppBackend/Private/BlueprintCompilerCppBackendValueHelper.cpp
--- a/Engine/Source/Developer/BlueprintCompilerCppBackend/Private/BlueprintCompilerCppBackendValueHelper.cpp
+++ b/Engine/Source/Developer/BlueprintCompilerCppBackend/Private/BlueprintCompilerCppBackendValueHelper.cpp
@@ -49,6 +49,7 @@ void FEmitDefaultValueHelper::OuterGenerate(FEmitterLocalContext& Context
 	for (int32 ArrayIndex = 0; ArrayIndex < Property->ArrayDim; ++ArrayIndex)
 	{
 		if (!OptionalDefaultDataContainer
+			|| ((Property->PropertyFlags & CPF_Config) != 0) 
 			|| (!Property->Identical_InContainer(DataContainer, OptionalDefaultDataContainer, ArrayIndex) && !IsInstancedSubobjectLambda(ArrayIndex)))
 		{
 			FString PathToMember;
@@ -194,9 +195,13 @@ void FEmitDefaultValueHelper::InnerGenerate(FEmitterLocalContext& Context, const
 	
 	if (ArrayProperty)
 	{
-		const bool bInitializeWithoutScriptStruct = false;
 		FScriptArrayHelper ScriptArrayHelper(ArrayProperty, ValuePtr);
 		UStructProperty* InnerStructProperty = Cast<UStructProperty>(ArrayProperty->Inner);
+		
+		const bool bNoExportStructure = InnerStructProperty && InnerStructProperty->Struct
+			&& InnerStructProperty->Struct->IsNative()
+			&& (InnerStructProperty->Struct->StructFlags & STRUCT_NoExport);
+		const bool bInitializeWithoutScriptStruct = bNoExportStructure;
 
 		UScriptStruct* RegularInnerStruct = nullptr;
 		if (!bInitializeWithoutScriptStruct)
@@ -243,7 +248,7 @@ void FEmitDefaultValueHelper::InnerGenerate(FEmitterLocalContext& Context, const
 	}
 }
 
-bool FEmitDefaultValueHelper::SpecialStructureConstructor(const UScriptStruct* Struct, const uint8* ValuePtr, /*out*/ FString* OutResult)
+bool FEmitDefaultValueHelper::SpecialStructureConstructor(const UStruct* Struct, const uint8* ValuePtr, /*out*/ FString* OutResult)
 {
 	check(ValuePtr || !OutResult);
 
@@ -701,7 +706,7 @@ public:
 		return FString(TEXT("Z_Construct_")) + ClassString + TEXT("_") + Result + TEXT("()");
 	}
 
-	static void AddStaticFunctionsForDependencies(FEmitterLocalContext& Context)
+	static void AddStaticFunctionsForDependencies(FEmitterLocalContext& Context, TSharedPtr<FGatherConvertedClassDependencies> ParentDependencies)
 	{
 		auto SourceClass = Context.GetCurrentlyGeneratedClass();
 		auto OriginalClass = Context.Dependencies.FindOriginalClass(SourceClass);
@@ -712,6 +717,13 @@ public:
 		Context.AddLine(TEXT("{"));
 		Context.IncreaseIndent();
 
+		UBlueprintGeneratedClass* ParentBPGC = Cast<UBlueprintGeneratedClass>(SourceClass->GetSuperClass());
+		ensure(ParentDependencies.IsValid() == !!ParentBPGC);
+		if (ParentBPGC)
+		{
+			Context.AddLine(FString::Printf(TEXT("%s::__StaticDependenciesAssets(AssetsToLoad);"), *FEmitHelper::GetCppName(ParentBPGC)));
+		}
+
 		auto CreateAssetToLoadString = [&](const UObject* AssetObj) -> FString
 		{
 			UClass* AssetType = AssetObj->GetClass();
@@ -736,10 +748,18 @@ public:
 		};
 		for (UObject* LocAsset : Context.Dependencies.Assets)
 		{
+			if (ParentDependencies.IsValid() && ParentDependencies->Assets.Contains(LocAsset))
+			{
+				continue;
+			}
 			Context.AddLine(CreateAssetToLoadString(LocAsset));
 		}
-		for (UObject* LocAsset : Context.Dependencies.ConvertedClasses)
+		for (UBlueprintGeneratedClass* LocAsset : Context.Dependencies.ConvertedClasses)
 		{
+			if (ParentDependencies.IsValid() && (ParentDependencies->ConvertedClasses.Contains(LocAsset) || ParentDependencies->Assets.Contains(LocAsset)))
+			{
+				continue;
+			}
 			if (!Context.Dependencies.Assets.Contains(LocAsset))
 			{
 				Context.AddLine(CreateAssetToLoadString(LocAsset));
@@ -776,7 +796,27 @@ public:
 	}
 };
 
-void FEmitDefaultValueHelper::GenerateCustomDynamicClassInitialization(FEmitterLocalContext& Context)
+void FEmitDefaultValueHelper::GenerateCustomDynamicClassInitializationUsedAssets(FEmitterLocalContext& Context)
+{
+	auto BPGC = CastChecked<UBlueprintGeneratedClass>(Context.GetCurrentlyGeneratedClass());
+	const FString CppClassName = FEmitHelper::GetCppName(BPGC);
+
+	Context.AddLine(FString::Printf(TEXT("void %s::__CustomDynamicClassInitializationUsedAssets(UDynamicClass* InDynamicClass)"), *CppClassName));
+	Context.AddLine(TEXT("{"));
+	Context.IncreaseIndent();
+	Context.AddLine(TEXT("ensure(0 == InDynamicClass->UsedAssets.Num());"));
+
+	for (auto LocAsset : Context.UsedObjectInCurrentClass)
+	{
+		const FString AssetStr = Context.FindGloballyMappedObject(LocAsset, UObject::StaticClass(), true, false);
+		Context.AddLine(FString::Printf(TEXT("InDynamicClass->UsedAssets.Add(%s);"), *AssetStr));
+	}
+
+	Context.DecreaseIndent();
+	Context.AddLine(TEXT("}"));
+}
+
+void FEmitDefaultValueHelper::GenerateCustomDynamicClassInitialization(FEmitterLocalContext& Context, TSharedPtr<FGatherConvertedClassDependencies> ParentDependencies)
 {
 	auto BPGC = CastChecked<UBlueprintGeneratedClass>(Context.GetCurrentlyGeneratedClass());
 	const FString CppClassName = FEmitHelper::GetCppName(BPGC);
@@ -786,7 +826,6 @@ void FEmitDefaultValueHelper::GenerateCustomDynamicClassInitialization(FEmitterL
 	Context.IncreaseIndent();
 	Context.AddLine(TEXT("ensure(0 == InDynamicClass->ReferencedConvertedFields.Num());"));
 	Context.AddLine(TEXT("ensure(0 == InDynamicClass->MiscConvertedSubobjects.Num());"));
-	Context.AddLine(TEXT("ensure(0 == InDynamicClass->UsedAssets.Num());"));
 	Context.AddLine(TEXT("ensure(0 == InDynamicClass->DynamicBindingObjects.Num());"));
 	Context.AddLine(TEXT("ensure(0 == InDynamicClass->ComponentTemplates.Num());"));
 	Context.AddLine(TEXT("ensure(0 == InDynamicClass->Timelines.Num());"));
@@ -796,6 +835,16 @@ void FEmitDefaultValueHelper::GenerateCustomDynamicClassInitialization(FEmitterL
 	Context.CurrentCodeType = FEmitterLocalContext::EGeneratedCodeType::SubobjectsOfClass;
 	Context.ResetPropertiesForInaccessibleStructs();
 
+	if (Context.Dependencies.ConvertedEnum.Num())
+	{
+		Context.AddLine(TEXT("// List of all referenced converted enums"));
+	}
+	for (auto LocEnum : Context.Dependencies.ConvertedEnum)
+	{
+		Context.AddLine(FString::Printf(TEXT("InDynamicClass->ReferencedConvertedFields.Add(LoadObject<UEnum>(nullptr, TEXT(\"%s\")));"), *(LocEnum->GetPathName().ReplaceCharWithEscapedChar())));
+		Context.EnumsInCurrentClass.Add(LocEnum);
+	}
+
 	if (Context.Dependencies.ConvertedClasses.Num())
 	{
 		Context.AddLine(TEXT("// List of all referenced converted classes"));
@@ -805,6 +854,11 @@ void FEmitDefaultValueHelper::GenerateCustomDynamicClassInitialization(FEmitterL
 		UClass* ClassToLoad = Context.Dependencies.FindOriginalClass(LocStruct);
 		if (ensure(ClassToLoad))
 		{
+			if (ParentDependencies.IsValid() && ParentDependencies->ConvertedClasses.Contains(LocStruct))
+			{
+				continue;
+			}
+
 			const FString ClassConstructor = FDependenciesHelper::GenerateZConstructor(ClassToLoad);
 			Context.AddLine(FString::Printf(TEXT("extern UClass* %s;"), *ClassConstructor));
 			Context.AddLine(FString::Printf(TEXT("InDynamicClass->ReferencedConvertedFields.Add(%s);"), *ClassConstructor));
@@ -820,20 +874,15 @@ void FEmitDefaultValueHelper::GenerateCustomDynamicClassInitialization(FEmitterL
 	}
 	for (auto LocStruct : Context.Dependencies.ConvertedStructs)
 	{
+		if (ParentDependencies.IsValid() && ParentDependencies->ConvertedStructs.Contains(LocStruct))
+		{
+			continue;
+		}
 		const FString StructConstructor = FDependenciesHelper::GenerateZConstructor(LocStruct);
 		Context.AddLine(FString::Printf(TEXT("extern UScriptStruct* %s;"), *StructConstructor));
 		Context.AddLine(FString::Printf(TEXT("InDynamicClass->ReferencedConvertedFields.Add(%s);"), *StructConstructor));
 	}
 
-	if (Context.Dependencies.ConvertedEnum.Num())
-	{
-		Context.AddLine(TEXT("// List of all referenced converted enums"));
-	}
-	for (auto LocEnum : Context.Dependencies.ConvertedEnum)
-	{
-		Context.AddLine(FString::Printf(TEXT("InDynamicClass->ReferencedConvertedFields.Add(LoadObject<UEnum>(nullptr, TEXT(\"%s\")));"), *(LocEnum->GetPathName().ReplaceCharWithEscapedChar())));
-	}
-
 	TArray<UActorComponent*> ActorComponentTempatesOwnedByClass = BPGC->ComponentTemplates;
 	// Gather all CT from SCS and IH, the remaining ones are generated for class..
 	if (auto SCS = BPGC->SimpleConstructionScript)
@@ -853,14 +902,14 @@ void FEmitDefaultValueHelper::GenerateCustomDynamicClassInitialization(FEmitterL
 		});
 	}
 
-	if (Context.Dependencies.Assets.Num())
-	{
-		Context.AddLine(TEXT("// List of all referenced assets"));
-	}
+	Context.AddLine(TEXT("__CustomDynamicClassInitializationUsedAssets(InDynamicClass);"));
+
 	for (auto LocAsset : Context.Dependencies.Assets)
 	{
-		const FString AssetStr = Context.FindGloballyMappedObject(LocAsset, UObject::StaticClass(), true, false);
-		Context.AddLine(FString::Printf(TEXT("InDynamicClass->UsedAssets.Add(%s);"), *AssetStr));
+		if (!ParentDependencies.IsValid() || !ParentDependencies->Assets.Contains(LocAsset))
+		{
+			Context.UsedObjectInCurrentClass.Add(LocAsset);
+		}
 	}
 
 	auto CreateAndInitializeClassSubobjects = [&](bool bCreate, bool bInitilize)
@@ -901,20 +950,24 @@ void FEmitDefaultValueHelper::GenerateCustomDynamicClassInitialization(FEmitterL
 
 void FEmitDefaultValueHelper::GenerateConstructor(FEmitterLocalContext& Context)
 {
-	GenerateCustomDynamicClassInitialization(Context);
+	UBlueprintGeneratedClass* BPGC = CastChecked<UBlueprintGeneratedClass>(Context.GetCurrentlyGeneratedClass());
+	UBlueprintGeneratedClass* ParentBPGC = Cast<UBlueprintGeneratedClass>(BPGC->GetSuperClass());
+	TSharedPtr<FGatherConvertedClassDependencies> ParentDependencies(ParentBPGC ? new FGatherConvertedClassDependencies(ParentBPGC) : nullptr);
 
-	auto BPGC = CastChecked<UBlueprintGeneratedClass>(Context.GetCurrentlyGeneratedClass());
+	GenerateCustomDynamicClassInitialization(Context, ParentDependencies);
+
+	
 	const FString CppClassName = FEmitHelper::GetCppName(BPGC);
 
 	UClass* SuperClass = BPGC->GetSuperClass();
-	const bool bSuperHasOnlyDefaultConstructor = SuperClass && SuperClass->HasMetaData(TEXT("OnlyDefaultConstructorDeclared"));
+	const bool bSuperHasObjectInitializerConstructor = SuperClass && SuperClass->HasMetaData(TEXT("ObjectInitializerConstructorDeclared"));
 
 	Context.CurrentCodeType = FEmitterLocalContext::EGeneratedCodeType::CommonConstructor;
 	Context.ResetPropertiesForInaccessibleStructs();
 	Context.AddLine(FString::Printf(TEXT("%s::%s(const FObjectInitializer& ObjectInitializer) : Super(%s)")
 		, *CppClassName
 		, *CppClassName
-		, bSuperHasOnlyDefaultConstructor ? TEXT("") : TEXT("ObjectInitializer")));
+		, bSuperHasObjectInitializerConstructor ? TEXT("ObjectInitializer") : TEXT("")));
 	Context.AddLine(TEXT("{"));
 	Context.IncreaseIndent();
 
@@ -1076,7 +1129,7 @@ void FEmitDefaultValueHelper::GenerateConstructor(FEmitterLocalContext& Context)
 	Context.DecreaseIndent();
 	Context.AddLine(TEXT("}"));
 
-	FDependenciesHelper::AddStaticFunctionsForDependencies(Context);
+	FDependenciesHelper::AddStaticFunctionsForDependencies(Context, ParentDependencies);
 
 	FBackendHelperUMG::EmitWidgetInitializationFunctions(Context);
 }
@@ -1115,7 +1168,7 @@ FString FEmitDefaultValueHelper::HandleClassSubobject(FEmitterLocalContext& Cont
 			, *NativeType
 			, bAddAsSubobjectOfClass ? TEXT("InDynamicClass") : *OuterStr
 			, *ActualClass
-			, *Object->GetName()));
+			, *Object->GetName().ReplaceCharWithEscapedChar()));
 		if (bAddAsSubobjectOfClass)
 		{
 			Context.RegisterClassSubobject(Object, ListOfSubobjectsType);
@@ -1226,7 +1279,7 @@ FString FEmitDefaultValueHelper::HandleInstancedSubobject(FEmitterLocalContext&
 			, *NativeType
 			, *OuterStr
 			, *ActualClass
-			, *Object->GetName()));
+			, *Object->GetName().ReplaceCharWithEscapedChar()));
 	}
 
 	return LocalNativeName;
diff --git a/Engine/Source/Developer/BlueprintNativeCodeGen/Private/BlueprintNativeCodeGenUtils.cpp b/Engine/Source/Developer/BlueprintNativeCodeGen/Private/BlueprintNativeCodeGenUtils.cpp
--- a/Engine/Source/Developer/BlueprintNativeCodeGen/Private/BlueprintNativeCodeGenUtils.cpp
+++ b/Engine/Source/Developer/BlueprintNativeCodeGen/Private/BlueprintNativeCodeGenUtils.cpp
@@ -133,6 +133,10 @@ static bool BlueprintNativeCodeGenUtilsImpl::GenerateModuleSourceFiles(const FBl
 	PchIncludes.Add(EngineHeaderFile);
 	PchIncludes.Add(TEXT("GeneratedCodeHelpers.h"));
 
+	TArray<FString> FilesToIncludeInModuleHeader;
+	GConfig->GetArray(TEXT("BlueprintNativizationSettings"), TEXT("FilesToIncludeInModuleHeader"), FilesToIncludeInModuleHeader, GEditorIni);
+	PchIncludes.Append(FilesToIncludeInModuleHeader);
+
 	bool bSuccess = GameProjectUtils::GeneratePluginModuleHeaderFile(TargetPaths.RuntimeModuleFile(FBlueprintNativeCodeGenPaths::HFile), PchIncludes, FailureReason);
 
 	if (bSuccess)
@@ -282,7 +286,8 @@ void FBlueprintNativeCodeGenUtils::GenerateCppCode(UObject* Obj, TSharedPtr<FStr
 		{
 			FBlueprintDuplicationScopeFlags BPDuplicationFlags(FBlueprintDuplicationScopeFlags::NoExtraCompilation 
 				| FBlueprintDuplicationScopeFlags::TheSameTimelineGuid 
-				| FBlueprintDuplicationScopeFlags::ValidatePinsUsingSourceClass);
+				| FBlueprintDuplicationScopeFlags::ValidatePinsUsingSourceClass
+				| FBlueprintDuplicationScopeFlags::TheSameNodeGuid);
 			DuplicateBP = DuplicateObject<UBlueprint>(InBlueprintObj, TempPackage, *InBlueprintObj->GetName());
 		}
 		ensure((nullptr != DuplicateBP->GeneratedClass) && (InBlueprintObj->GeneratedClass != DuplicateBP->GeneratedClass));
diff --git a/Engine/Source/Developer/BlueprintProfiler/Private/BlueprintProfiler.cpp b/Engine/Source/Developer/BlueprintProfiler/Private/BlueprintProfiler.cpp
--- a/Engine/Source/Developer/BlueprintProfiler/Private/BlueprintProfiler.cpp
+++ b/Engine/Source/Developer/BlueprintProfiler/Private/BlueprintProfiler.cpp
@@ -2,6 +2,7 @@
 
 #include "BlueprintProfilerPCH.h"
 #include "EditorStyleSet.h"
+#include "ActorEditorUtils.h"
 #include "Editor/UnrealEd/Classes/Settings/EditorExperimentalSettings.h"
 
 #define LOCTEXT_NAMESPACE "BlueprintProfiler"
@@ -25,9 +26,10 @@ FBlueprintProfiler::FBlueprintProfiler()
 	// Update stat watermarks from editor settings, the callbacks should keep them in sync after this.
 	if (const UEditorExperimentalSettings* EditorSettings = GetDefault<UEditorExperimentalSettings>())
 	{
+		FScriptPerfData::EnableRecentSampleBias(EditorSettings->bEnableBlueprintProfilerRecentSampleBias);
 		FScriptPerfData::SetRecentSampleBias(EditorSettings->BlueprintProfilerRecentSampleBias);
 		FScriptPerfData::SetEventPerformanceThreshold(EditorSettings->BlueprintProfilerEventThreshold);
-		FScriptPerfData::SetNodePerformanceThreshold(EditorSettings->BlueprintProfilerExclNodeThreshold);
+		FScriptPerfData::SetExclusivePerformanceThreshold(EditorSettings->BlueprintProfilerExclNodeThreshold);
 		FScriptPerfData::SetInclusivePerformanceThreshold(EditorSettings->BlueprintProfilerInclNodeThreshold);
 		FScriptPerfData::SetMaxPerformanceThreshold(EditorSettings->BlueprintProfilerMaxNodeThreshold);
 	}
@@ -70,7 +72,7 @@ void FBlueprintProfiler::InstrumentEvent(const EScriptInstrumentationEvent& Even
 	SCOPE_CYCLE_COUNTER(STAT_ProfilerInstrumentationCost);
 
 	const UObject* ContextObject = Event.GetContextObject();
-	if (ContextObject && !ContextObject->HasAnyFlags(RF_Transient))
+	if (ContextObject && !ContextObject->HasAnyFlags(RF_Transient) && !FActorEditorUtils::IsAPreviewOrInactiveActor(Cast<const AActor>(ContextObject)))
 	#endif
 	{
 		// Handle context switching events
@@ -105,8 +107,12 @@ TSharedPtr<FBlueprintExecutionContext> FBlueprintProfiler::GetBlueprintContext(c
 		// Map the blueprint and initialise the context.
 		if (Result->InitialiseContext(BlueprintClassPath))
 		{
-			// Register delegate to handle updates
-			Result->GetBlueprint().Get()->OnCompiled().AddRaw(this, &FBlueprintProfiler::RemoveAllBlueprintReferences);
+			if (UBlueprint* Blueprint = Result->GetBlueprint().Get())
+			{
+				// Register delegate to handle updates
+				Blueprint->OnCompiled().AddRaw(this, &FBlueprintProfiler::RemoveAllBlueprintReferences);
+				GraphLayoutChangedDelegate.Broadcast(Blueprint);
+			}
 		}
 		else
 		{
@@ -199,13 +205,15 @@ void FBlueprintProfiler::ProcessEventProfilingData()
 						NewEventRange.BlueprintContext = GetBlueprintContext(CurrEvent.GetObjectPath());
 						if (NewEventRange.BlueprintContext.IsValid())
 						{
-							NewEventRange.InstanceName = MapBlueprintInstance(NewEventRange.BlueprintContext, InstanceEvent.GetObjectPath());
+							NewEventRange.InstanceName = NewEventRange.BlueprintContext->MapBlueprintInstance(InstanceEvent.GetObjectPath());
 							NewEventRange.StartIdx = EventIdx;
 							ScriptEventRanges.Push(NewEventRange);
 						}
+
+						// Only remove the class/instance pair for new events.
+						InstrumentationEventQueue.RemoveAt(EventIdx, 2, false);
 					}
 				}
-				InstrumentationEventQueue.RemoveAt(EventIdx, 2, false);
 				break;
 			}
 			case EScriptInstrumentation::Event:
@@ -353,34 +361,4 @@ void FBlueprintProfiler::EndPIE(bool bIsSimulating)
 #endif // WITH_EDITOR
 }
 
-FName FBlueprintProfiler::MapBlueprintInstance(TSharedPtr<FBlueprintExecutionContext> BlueprintContext, const FString& InstancePath)
-{
-	FName InstanceName(*InstancePath);
-	TWeakObjectPtr<const UObject> Instance = BlueprintContext->GetInstance(InstanceName);
-
-	if (Instance.IsValid() && !BlueprintContext->HasProfilerDataForInstance(InstanceName))
-	{
-		// Create new instance node
-		FScriptExecNodeParams InstanceNodeParams;
-		InstanceNodeParams.NodeName = InstanceName;
-		InstanceNodeParams.ObservedObject = Instance.Get();
-		InstanceNodeParams.NodeFlags = EScriptExecutionNodeFlags::Instance;
-		const AActor* Actor = Cast<AActor>(Instance.Get());
-		InstanceNodeParams.DisplayName = Actor ? FText::FromString(Actor->GetActorLabel()) : FText::FromString(Instance.Get()->GetName());
-		InstanceNodeParams.Tooltip = LOCTEXT("NavigateToInstanceHyperlink_ToolTip", "Navigate to the Instance");
-		InstanceNodeParams.IconColor = FLinearColor(1.f, 1.f, 1.f, 0.8f);
-		InstanceNodeParams.Icon = const_cast<FSlateBrush*>(FEditorStyle::GetBrush(TEXT("BlueprintProfiler.Actor")));
-		TSharedPtr<FScriptExecutionInstance> InstanceNode = MakeShareable<FScriptExecutionInstance>(new FScriptExecutionInstance(InstanceNodeParams));
-		// Link to parent blueprint entry
-		TSharedPtr<FScriptExecutionBlueprint> BlueprintNode = BlueprintContext->GetBlueprintExecNode();
-		BlueprintNode->AddInstance(InstanceNode);
-		// Fill out events from the blueprint root node
-		for (int32 NodeIdx = 0; NodeIdx < BlueprintNode->GetNumChildren(); ++NodeIdx)
-		{
-			InstanceNode->AddChildNode(BlueprintNode->GetChildByIndex(NodeIdx));
-		}
-	}
-	return InstanceName;
-}
-
 #undef LOCTEXT_NAMESPACE
diff --git a/Engine/Source/Developer/BlueprintProfiler/Private/BlueprintProfilerConnectionDrawingPolicy.cpp b/Engine/Source/Developer/BlueprintProfiler/Private/BlueprintProfilerConnectionDrawingPolicy.cpp
--- a/Engine/Source/Developer/BlueprintProfiler/Private/BlueprintProfilerConnectionDrawingPolicy.cpp
+++ b/Engine/Source/Developer/BlueprintProfiler/Private/BlueprintProfilerConnectionDrawingPolicy.cpp
@@ -61,7 +61,7 @@ void FBlueprintProfilerConnectionDrawingPolicy::BuildExecutionRoadmap()
 	FBlueprintDebugData& DebugData = TargetClass->GetDebugData();
 	IBlueprintProfilerInterface& ProfilerModule = FModuleManager::LoadModuleChecked<IBlueprintProfilerInterface>("BlueprintProfiler");
 	TSharedPtr<FBlueprintExecutionContext> BlueprintContext = ProfilerModule.GetBlueprintContext(TargetClass->GetPathName());
-	const FName InstanceName = ProfilerModule.MapBlueprintInstance(BlueprintContext, ActiveObject->GetPathName());
+	const FName InstanceName = BlueprintContext->MapBlueprintInstance(ActiveObject->GetPathName());
 	const FName GraphName = GraphObj->GetFName();
 	TSharedPtr<FBlueprintFunctionContext> FunctionContext = BlueprintContext->GetFunctionContext(GraphName);
 	const FName FunctionName = FunctionContext->GetFunctionName();
@@ -78,10 +78,10 @@ void FBlueprintProfilerConnectionDrawingPolicy::BuildExecutionRoadmap()
 	TArray<UEdGraphPin*> SequentialExecPinsInGraph;
 	TArray<FTracePath> SequentialTracePathsInGraph;
 
-	const TSimpleRingBuffer<FBlueprintExecutionTrace>& TraceStack = BlueprintContext->GetTraceHistory();
-	for (int32 i = 0; i < TraceStack.Num(); ++i)
+	const TSimpleRingBuffer<FBlueprintExecutionTrace>& TraceHistory = BlueprintContext->GetTraceHistory();
+	for (int32 i = 0; i < TraceHistory.Num(); ++i)
 	{
-		const FBlueprintExecutionTrace& Sample = TraceStack(i);
+		const FBlueprintExecutionTrace& Sample = TraceHistory(i);
 		if (InstanceName == Sample.InstanceName && FunctionName == Sample.FunctionName)
 		{
 			TSharedPtr<FScriptExecutionNode> ProfilerNode = FunctionContext->GetProfilerDataForNode(Sample.NodeName);
@@ -140,46 +140,72 @@ void FBlueprintProfilerConnectionDrawingPolicy::BuildExecutionRoadmap()
 				UEdGraphNode* NextNode = InputPin->GetOwningNode();
 				if (CurNode != NextNode)
 				{
-					TSharedPtr<FScriptExecutionNode> CurExecNode = BlueprintContext->GetProfilerDataForNode(CurNode);
+					TSharedPtr<FScriptExecutionNode> CurExecNode = BlueprintContext->GetProfilerDataForPin(OutputPin);
 					TSharedPtr<FScriptExecutionNode> NextExecNode = BlueprintContext->GetProfilerDataForNode(NextNode);
-					double NextNodeTime = SequentialNodeTimes[i];
-
-					FExecPairingMap& ExecPaths			= PredecessorPins.FindOrAdd(NextNode);
-					FTimePair& ExecTiming				= ExecPaths.FindOrAdd(OutputPin);
-					FProfilerPairingMap& ProfilerPaths	= ProfilerPins.FindOrAdd(NextNode);
-					FProfilerPair& ProfilerData			= ProfilerPaths.FindOrAdd(OutputPin);
-					// make sure that if we've already visited this exec-pin (like 
-					// in a for-loop or something), that we're replacing it with a 
-					// more recent execution time
-					//
-					// @TODO I don't see when this wouldn't be the case
-					if (ExecTiming.ThisExecTime < NextNodeTime)
+
+					if (CurExecNode.IsValid() && NextExecNode.IsValid())
 					{
-						double CurNodeTime = SequentialNodeTimes[OutputPinIndex];
-						ExecTiming.ThisExecTime = NextNodeTime;
-						ExecTiming.PredExecTime = CurNodeTime;
-						TSharedPtr<FScriptPerfData> CurPerfData = CurExecNode->GetPerfDataByInstanceAndTracePath(InstanceName, SequentialTracePathsInGraph[OutputPinIndex]);
-						TSharedPtr<FScriptPerfData> NextPerfData = NextExecNode->GetPerfDataByInstanceAndTracePath(InstanceName, SequentialTracePathsInGraph[i]);
-						ProfilerData.bPureNode = NextExecNode->IsPureNode();
-						switch(WireHeatMode)
+						double NextNodeTime = SequentialNodeTimes[i];
+
+						FExecPairingMap& ExecPaths			= PredecessorPins.FindOrAdd(NextNode);
+						FTimePair& ExecTiming				= ExecPaths.FindOrAdd(OutputPin);
+						FProfilerPairingMap& ProfilerPaths	= ProfilerPins.FindOrAdd(NextNode);
+						FProfilerPair& ProfilerData			= ProfilerPaths.FindOrAdd(OutputPin);
+						// make sure that if we've already visited this exec-pin (like 
+						// in a for-loop or something), that we're replacing it with a 
+						// more recent execution time
+						//
+						// @TODO I don't see when this wouldn't be the case
+						if (ExecTiming.ThisExecTime < NextNodeTime)
 						{
-							case EBlueprintProfilerHeatMapDisplayMode::Exclusive:
-							{
-								ProfilerData.PredPerfData = CurPerfData->GetNodeHeatLevel();
-								ProfilerData.ThisPerfData = NextPerfData->GetNodeHeatLevel();
-								break;
-							}
-							case EBlueprintProfilerHeatMapDisplayMode::Inclusive:
-							{
-								ProfilerData.PredPerfData = CurPerfData->GetInclusiveHeatLevel();
-								ProfilerData.ThisPerfData = NextPerfData->GetInclusiveHeatLevel();
-								break;
-							}
-							case EBlueprintProfilerHeatMapDisplayMode::MaxTiming:
+							double CurNodeTime = SequentialNodeTimes[OutputPinIndex];
+							ExecTiming.ThisExecTime = NextNodeTime;
+							ExecTiming.PredExecTime = CurNodeTime;
+							const bool bExecPin = CurExecNode->HasFlags(EScriptExecutionNodeFlags::ExecPin);
+							const FTracePath& CurTracePath = bExecPin ? SequentialTracePathsInGraph[i] : SequentialTracePathsInGraph[OutputPinIndex];
+							const FTracePath& NextTracePath = SequentialTracePathsInGraph[i];
+							TSharedPtr<FScriptPerfData> CurPerfData = CurExecNode->GetPerfDataByInstanceAndTracePath(InstanceName, CurTracePath);
+							TSharedPtr<FScriptPerfData> NextPerfData = NextExecNode->GetPerfDataByInstanceAndTracePath(InstanceName, NextTracePath);
+							ProfilerData.bPureNode = CurExecNode->IsPureNode();
+							switch(WireHeatMode)
 							{
-								ProfilerData.PredPerfData = CurPerfData->GetMaxTimeHeatLevel();
-								ProfilerData.ThisPerfData = NextPerfData->GetMaxTimeHeatLevel();
-								break;
+								case EBlueprintProfilerHeatMapDisplayMode::Exclusive:
+								{
+									ProfilerData.PredPerfData = CurPerfData->GetNodeHeatLevel();
+									ProfilerData.ThisPerfData = NextPerfData->GetNodeHeatLevel();
+									break;
+								}
+								case EBlueprintProfilerHeatMapDisplayMode::Inclusive:
+								case EBlueprintProfilerHeatMapDisplayMode::PinToPin:
+								{
+									ProfilerData.PredPerfData = CurPerfData->GetInclusiveHeatLevel();
+									ProfilerData.ThisPerfData = NextPerfData->GetInclusiveHeatLevel();
+									break;
+								}
+								case EBlueprintProfilerHeatMapDisplayMode::MaxTiming:
+								{
+									ProfilerData.PredPerfData = CurPerfData->GetMaxTimeHeatLevel();
+									ProfilerData.ThisPerfData = NextPerfData->GetMaxTimeHeatLevel();
+									break;
+								}
+								case EBlueprintProfilerHeatMapDisplayMode::Total:
+								{
+									ProfilerData.PredPerfData = CurPerfData->GetTotalHeatLevel();
+									ProfilerData.ThisPerfData = NextPerfData->GetTotalHeatLevel();
+									break;
+								}
+								case EBlueprintProfilerHeatMapDisplayMode::HottestPath:
+								{
+									ProfilerData.PredPerfData = CurPerfData->GetHottestPathHeatLevel();
+									ProfilerData.ThisPerfData = NextPerfData->GetHottestPathHeatLevel();
+									break;
+								}
+								case EBlueprintProfilerHeatMapDisplayMode::HottestEndpoint:
+								{
+									ProfilerData.PredPerfData = CurPerfData->GetHottestEndpointHeatLevel();
+									ProfilerData.ThisPerfData = NextPerfData->GetHottestEndpointHeatLevel();
+									break;
+								}
 							}
 						}
 					}
@@ -197,12 +223,12 @@ void FBlueprintProfilerConnectionDrawingPolicy::DeterminePerfWiringStyle(UEdGrap
 	// Initiate using kismet connection drawing policy
 	FKismetConnectionDrawingPolicy::DetermineWiringStyle(OutputPin, InputPin, Params);
 	// Override the colors using the profiler Data
-	Params.WireColor = Params.WireColor2 = FLinearColor::White;
+	static const FLinearColor BaseColor(1.f, 1.f, 1.f, 0.2f);
+	Params.WireColor = Params.WireColor2 = BaseColor;
 	if (FProfilerPairingMap* ExecMap = ProfilerPins.Find(InputPin->GetOwningNode()))
 	{
 		if (FProfilerPair* PerfData = ExecMap->Find(OutputPin))
 		{
-			static const FLinearColor BaseColor(0.f, 0.f, 0.f, 1.f);
 			static const FLinearColor ImpureHeatColor(1.f, 0.f, 0.f, 1.f);
 			static const FLinearColor PureHeatColor(0.f, 1.f, 0.f, 1.f);
 			Params.PerformanceData1 = FMath::Clamp(PerfData->PredPerfData, 0.f, 1.f);
@@ -399,7 +425,6 @@ void FBlueprintProfilerConnectionDrawingPolicy::DrawPerfConnection(int32 LayerId
 				if (Distance < SplineLength)
 				{
 					const float Alpha = SplineReparamTable.Eval(Distance, 0.f);
-					const float AlphaSq = Alpha*Alpha;
 					FVector2D BubblePos = FMath::CubicInterp(P0, P0Tangent, P1, P1Tangent, Alpha);
 					const FVector2D BubbleSize = BubbleImage->ImageSize * FMath::Lerp(SizeA, SizeB, Alpha);
 					const FLinearColor ElementColor = FLinearColor::LerpUsingHSV(Params.WireColor, Params.WireColor2, Alpha);
