Fixing ambiguous logical operation (causing CIS error on Linux).

diff --git a/Engine/Source/Developer/BlueprintCompilerCppBackend/Private/BlueprintCompilerCppBackend.cpp b/Engine/Source/Developer/BlueprintCompilerCppBackend/Private/BlueprintCompilerCppBackend.cpp
--- a/Engine/Source/Developer/BlueprintCompilerCppBackend/Private/BlueprintCompilerCppBackend.cpp
+++ b/Engine/Source/Developer/BlueprintCompilerCppBackend/Private/BlueprintCompilerCppBackend.cpp
@@ -632,9 +632,9 @@ FString FBlueprintCompilerCppBackend::EmitCallStatmentInner(FEmitterLocalContext
 			auto FunctionOwner = Statement.FunctionToCall->GetOwnerClass();
 			auto OwnerBPGC = Cast<UBlueprintGeneratedClass>(FunctionOwner);
 			const bool bUnconvertedClass = OwnerBPGC && !EmitterContext.Dependencies.WillClassBeConverted(OwnerBPGC);
-			const bool bIsCustomThunk = bStaticCall && Statement.FunctionToCall->GetBoolMetaData(TEXT("CustomThunk"))
+			const bool bIsCustomThunk = bStaticCall && ( Statement.FunctionToCall->GetBoolMetaData(TEXT("CustomThunk"))
 				|| Statement.FunctionToCall->HasMetaData(TEXT("CustomStructureParam"))
-				|| Statement.FunctionToCall->HasMetaData(TEXT("ArrayParm"));
+				|| Statement.FunctionToCall->HasMetaData(TEXT("ArrayParm")) );
 			if (bUnconvertedClass)
 			{
 				ensure(!Statement.bIsParentContext); //unsupported yet
