Copying //UE4/Dev-Blueprints to //UE4/Dev-Main (Source: //UE4/Dev-Blueprints @ 3042808)

==========================
MAJOR FEATURES + CHANGES
==========================

Change 3026763 on 2016/06/24 by [REDACTED_PERSON]

	[UE-32259] Fix missing heat map coloration on pure nodes.

	change summary:
	- modified FScriptExecutionNode::RefreshStats() to ensure that child nodes are updated for 'standard' exec node types.

Change 3026829 on 2016/06/24 by [REDACTED_PERSON]

	Mirroring CL 3025832 from Release-4.12

	Guarding GetClassDefault nodes against asserts while loading cyclically dependent content. Regenerating one Blueprint could trigger a ReconstructNode() in a separate Blueprint (that had already been regenerated). However, not all objects in that Blueprint's package have been post-loaded, so object Rename() will assert is you try to reset loaders with it.

Change 3026888 on 2016/06/24 by [REDACTED_PERSON]

	PR #2460: Functionality to flag blueprint as abstract (Contributed by [REDACTED_PERSON])

Change 3027245 on 2016/06/24 by [REDACTED_PERSON]

	[UEBP-180] Alternate treatment for Blueprint profiler tree view node expander icon (to differentiate from exec pin icons).

	change summary:
	- added SProfilerStatExpander (from BenC)
	- modified original to replace indicator icons w/ FontAwesome glyphs (they're a bit clearer) - see screenshot attached to JIRA

Change 3027279 on 2016/06/24 by [REDACTED_PERSON]

	Fixed issue where unsupported pins could be added to custom Event nodes which can lead to a crash

Change 3027606 on 2016/06/25 by [REDACTED_PERSON]

	Back out changelist 3027279, this change breaks custom events coupled to multicast delegate nodes (see ProjRocket.uasset in shooter game for an example)

Change 3028731 on 2016/06/27 by [REDACTED_PERSON]

	Minor change to the spacing around the new expander icons for the blueprint profiler.
	#Proj Kismet

Change 3028740 on 2016/06/27 by [REDACTED_PERSON]

	This change modifies the way the blueprint profiler detects and maps delegate events that are created at blueprint compile time.
	#Proj BlueprintProfiler, KismetCompiler, BlueprintGraph, LiveEditor, Kismet, Engine

Change 3028844 on 2016/06/27 by [REDACTED_PERSON]

	Better way to find Placeholder classes in BlueprintNativeCodeGenModule.

Change 3029573 on 2016/06/27 by [REDACTED_PERSON]

	Fixed issue where Preview Connection would be drawn incorrectly when an input and output pin share the same name

Change 3030335 on 2016/06/28 by [REDACTED_PERSON]

	List of AssetRegistrySearchable properties, that should be stored for DynamicClasses in moved to .ini file.
	Added SearchDynamicClassCues function in GameplayCueManager.cpp

Change 3030492 on 2016/06/28 by [REDACTED_PERSON]

	CIS Mac Build fix

Change 3030494 on 2016/06/28 by [REDACTED_PERSON]

	This change reworks the way the blueprint profiler maps latent re-entry points to address problems that could be triggered when multiple re-entries happenned within close proximity.
	#Proj BlueprintProfiler, CoreUObject, Engine, Kismet

Change 3030612 on 2016/06/28 by [REDACTED_PERSON]

	typo fixed

Change 3031461 on 2016/06/28 by [REDACTED_PERSON]

	Fixed for crash involving a stale UEdGraphPin after performing an undo/redo operation inolving UEdGraphNode(s)

Change 3031537 on 2016/06/28 by [REDACTED_PERSON]

	Fix for a problem introduced in CL 3030494 and recent graph pin changes that caused problems uniquely naming tunnel entry/exit pins.
	#Proj BlueprintProfiler

Change 3031556 on 2016/06/28 by [REDACTED_PERSON]

	This fixes a problem with the profiler tracking of PIE actor instances adding a mechanism to refresh the actor references when a subsequent PIE session is started.
	#Proj BlueprintProfiler, Kismet

Change 3031588 on 2016/06/28 by [REDACTED_PERSON]

	Fixed issue where unsupported pins could be added to custom Event nodes which can lead to a crash

Change 3031887 on 2016/06/28 by [REDACTED_PERSON]

	Fix for crash caused by stale pin access in SGraphPin, also fixes 'double opacity' effect when ctrl+dragging off of a macro

Change 3031932 on 2016/06/28 by [REDACTED_PERSON]

	Fix for major regression in pin value visibility

Change 3032568 on 2016/06/29 by [REDACTED_PERSON]

	[UE-31681] Ensure that cached component data cache template objects are referenced during garbage collection.

	change summary:
	- modified AActor::AddReferencedObjects() to include CurrentTransactionAnnotation
	- added UChildActorComponent::AddReferencedObjects() to include CachedInstanceData
	- also added FChildActorComponentInstanceData::AddReferencedObjects() to include ComponentInstanceData (just in case some future task happens to trigger a reference gather during its lifetime)

	additional notes:
	- should also resolve UE-32558

Change 3032998 on 2016/06/29 by [REDACTED_PERSON]

	Fixing shutdown crash after LinkedTo lists are corrupted. An ensure couldn't fire when a node had already been destroyed. UE-32631 tracks the actual ensure.

Change 3033448 on 2016/06/29 by [REDACTED_PERSON]

	Discard Return Nodes that have no outputs when collapsing to function

Change 3034407 on 2016/06/30 by [REDACTED_PERSON]

	[UE-32675] Fix no-editor compile issue.

Change 3034691 on 2016/06/30 by [REDACTED_PERSON]

	Fixed issue where Macro nodes with ambiguous links could be collapsed into a function

Change 3034701 on 2016/06/30 by [REDACTED_PERSON]

	Cleaned up UEdGraphSchema_K2::GetVariableTypeTree to use a filter rather than several booleans

Change 3035093 on 2016/06/30 by [REDACTED_PERSON]

	Temporary fix.

Change 3035440 on 2016/07/01 by [REDACTED_PERSON]

	Make sure, the CDO of DynamicClass (and its subobjects) have AsyncLoading flag cleared.

Change 3035442 on 2016/07/01 by [REDACTED_PERSON]

	Added UDynamicClass::FindArchetype

Change 3036233 on 2016/07/01 by [REDACTED_PERSON]

	Proper handling of null pin references in compiler results logs

Change 3036541 on 2016/07/01 by [REDACTED_PERSON]

	Fix for more assymetry in LinkedTo after undo/redo. This solution addresses both UE-32645 and UE-32553. The key to correct behavior is to explictly detect the case where only one of the two nodes in a peer (LinkedTo) relationship are in the transaction buffer. Added IsObjectTransacting to transaction interface to detect this.

Change 3036581 on 2016/07/02 by [REDACTED_PERSON]

	Non-editor fix

Change 3036632 on 2016/07/02 by [REDACTED_PERSON]

	Removed too strict check in BlueprintNativeCodeGenModule

Change 3036715 on 2016/07/02 by [REDACTED_PERSON]

	- Private bitfield can be used as LHS expression in generated code.
	- Refactoring. Introduced ENativizedTermUsage.

Change 3037014 on 2016/07/04 by [REDACTED_PERSON]

	Make sure that temporary variable used in ternary operator: "context ? context->variable : temp-variable", has exactly the same type as the wanted variable.

Change 3037300 on 2016/07/05 by [REDACTED_PERSON]

Change 3037303 on 2016/07/05 by [REDACTED_PERSON]

	Nativization - fixed hardcoded string.

Change 3037307 on 2016/07/05 by [REDACTED_PERSON]

	AutomationTool:
	Fixed issue: A nativized client was compiled with both client and server NativizedAssets plugins.

	Paths to nativized BP plugin is stored in a map.

Change 3037712 on 2016/07/05 by [REDACTED_PERSON]

	Adding a non-const version of MultiFindPointer() for TMultiMap.

Change 3037790 on 2016/07/05 by [REDACTED_PERSON]

	Fixing a bug where overridden components in Blueprints could load with cleared properties - this could happen in cyclically dependent parent/child Blueprints (where the child's component would be zeroed). In this case, the Blueprint's CDO would be deferred, so to solve this new issue we now defer the Blueprint's sub-object overrides as well.

Change 3038014 on 2016/07/05 by [REDACTED_PERSON]

	Refactor of the blueprint profiler tunnel mapping and playback to address issues blocking MVP
	#Proj Kismet, BlueprintProfiler, KismetCompiler, BlueprintGraph, UnrealEd

Change 3038533 on 2016/07/05 by [REDACTED_PERSON]

	PR #2558: Stopping Zoom-to-Fit On User Input (Fixed) (Contributed by [REDACTED_PERSON])

Change 3038608 on 2016/07/05 by [REDACTED_PERSON]

	Fix for crash when altering ConstructObjectFromClass nodes, removed now unused InvalidateAndTrash function from UEdGraphPin, removed subpin leak counter

Change 3038770 on 2016/07/05 by [REDACTED_PERSON]

	Hiding  TThreadSingleton<FBlueprintExceptionTracker>::Get() to fix clang for windows link error

Change 3038841 on 2016/07/05 by [REDACTED_PERSON]

	Simple null check to prevent crash after loading a blueprint to diff. The diff'd blueprint is ending up in other blueprint's DependentBlueprintsToRecompile, for now this is the safest fix

Change 3039179 on 2016/07/06 by [REDACTED_PERSON]

	Selected (in Level editor) instance components are still selected after they are converted to regular components.

Change 3039216 on 2016/07/06 by [REDACTED_PERSON]

	FIxed unsafe code.

Change 3040046 on 2016/07/06 by [REDACTED_PERSON]

	This fixes a problem found where tunnel instances inside tunnel graphs fail to map correctly in the blueprint profiler
	#Proj BlueprintProfiler

Change 3040061 on 2016/07/06 by [REDACTED_PERSON]

	fix for bad unshelve before submit in CL 3040046

	- Phillip already reviewed this version of the file and not what was actually submitted.

Change 3040346 on 2016/07/06 by [REDACTED_PERSON]

	Fix(?)  for warning C28182

Change 3040349 on 2016/07/06 by [REDACTED_PERSON]

	Updating test data

Change 3040473 on 2016/07/06 by [REDACTED_PERSON]

	Fix SMyBlueprint in blueprint diff tool so that kismet inspector is updated with the object selected in SMyBlueprint, needed to investigate UE-18200

Change 3040536 on 2016/07/06 by [REDACTED_PERSON]

	Call to APITestFunction

Change 3040565 on 2016/07/06 by [REDACTED_PERSON]

	Clone saved GeneratedClass when not doing COL. Groundwork to fix for UE-17268.

Change 3040603 on 2016/07/06 by [REDACTED_PERSON]

	PR #2539: Option to disable Grid in the blueprint and material editor and change grid colors (Contributed by [REDACTED_PERSON])

Change 3040798 on 2016/07/07 by [REDACTED_PERSON]

	Fixed warning C6011: Dereferencing NULL pointer 'Object'.
	Added some comments.

Change 3041185 on 2016/07/07 by [REDACTED_PERSON]

	This fixes problems with the blueprint profiler with name collisions for nodes in different graphs and improves nested tunnel detection.
	#Proj BlueprintProfiler, Kismet

	Notes:

	This attempts to continue to address fall out from CL 3040046 to improve stability.
	Previously execution nodes were patched into the function contexts that called a tunnel instance. However this proved problematic after spotting a rare case where node names collide between
	between event and tunnel graphs and caused havok in the node mapping. This change moves to an approach that patches in tunnel contexts and any dependent contexts and instead attempts
	to locate critical path nodes using a new node search system making use of these child function contexts.

	CL 3040046 also modified TunnelContext's to be actually more representative of a tunnel instance context, I have updated naming conventions to reflect this for clarity.
	The changes here also add a single function context for tunnel graphs which contain the bulk of the nodes in the graph and the tunnel instance contexts reference this and make use of it to create
	non boundary / tunnel instance nodes so they are located in a single place for all instances of that tunnel graph.

Change 3041603 on 2016/07/07 by [REDACTED_PERSON]

	Fix for check failure: PreviousCDO != nullptr, when loading blueprints that have circularly dependent interfaces

Change 3042058 on 2016/07/07 by [REDACTED_PERSON]

	Ensures that the correct ComponentClass is assigned to our template node. This pins are normally allocated lazilly in BlueprintActionFilterImpl::HasMatchingPin, so not a huge change

Change 3042468 on 2016/07/08 by [REDACTED_PERSON]

	Fixed Crash in UHT.
	Fixed "Ambiguous search" error, when delegate has owner class declared.

Change 3042739 on 2016/07/08 by [REDACTED_PERSON]

	Nativization. Included headers for native subobjects.

Change 3042747 on 2016/07/08 by [REDACTED_PERSON]

	Minor changes in Orion code, necessary to compile the project with nativized Blueprints

Change 3042758 on 2016/07/08 by [REDACTED_PERSON]

	"OrionGame.h" is included in NativizedAssets module.

diff --git a/Engine/Config/BaseEditor.ini b/Engine/Config/BaseEditor.ini
--- a/Engine/Config/BaseEditor.ini
+++ b/Engine/Config/BaseEditor.ini
@@ -477,6 +477,11 @@ bDontNativizeDataOnlyBP=false
 ; this flag can change BP behaviour, see UE-23629. Moreover it doesn't help in image size reduction.
 bUseLocalEventGraphVariables=false
 
++ClassesWithStaticSearchableValues=/Script/GameplayAbilities.GameplayCueNotify_Actor
++ClassesWithStaticSearchableValues=/Script/GameplayAbilities.GameplayCueNotify_Static
+
++StaticSearchableTagNames=GameplayCueName
+
 ; +ExcludedBlueprintTypes=/Script/UMGEditor.WidgetBlueprint
 ; +ExcludedBlueprintTypes=/Script/Engine.AnimBlueprint
 ; +ExcludedBlueprintTypes=/Script/GameplayAbilities.GameplayAbilityBlueprint
diff --git a/Engine/Plugins/Experimental/LiveEditor/Source/LiveEditor/Private/K2Node_LiveEditObject.cpp b/Engine/Plugins/Experimental/LiveEditor/Source/LiveEditor/Private/K2Node_LiveEditObject.cpp
--- a/Engine/Plugins/Experimental/LiveEditor/Source/LiveEditor/Private/K2Node_LiveEditObject.cpp
+++ b/Engine/Plugins/Experimental/LiveEditor/Source/LiveEditor/Private/K2Node_LiveEditObject.cpp
@@ -347,7 +347,7 @@ void UK2Node_LiveEditObject::ExpandNode(class FKismetCompilerContext& CompilerCo
 		//the event itself
 		//
 		UFunction *EventMIDISignature = GetEventMIDISignature();
-		UK2Node_Event* EventNode = CompilerContext.SpawnIntermediateNode<UK2Node_Event>(this, SourceGraph);
+		UK2Node_Event* EventNode = CompilerContext.SpawnIntermediateEventNode<UK2Node_Event>(this, nullptr, SourceGraph);
 		check(EventNode);
 		EventNode->EventReference.SetFromField<UFunction>(EventMIDISignature, false);
 		EventNode->CustomFunctionName = *EventNameGuid;
diff --git a/Engine/Source/Developer/BlueprintCompilerCppBackend/Private/BlueprintCompilerCppBackend.cpp b/Engine/Source/Developer/BlueprintCompilerCppBackend/Private/BlueprintCompilerCppBackend.cpp
--- a/Engine/Source/Developer/BlueprintCompilerCppBackend/Private/BlueprintCompilerCppBackend.cpp
+++ b/Engine/Source/Developer/BlueprintCompilerCppBackend/Private/BlueprintCompilerCppBackend.cpp
@@ -5,6 +5,7 @@
 #include "BlueprintCompilerCppBackendUtils.h"
 #include "EdGraphSchema_K2.h"
 #include "Kismet/KismetNodeHelperLibrary.h"
+#include "Kismet/KismetArrayLibrary.h"
 
 // Generates single "if" scope. Its condition checks context of given term.
 struct FSafeContextScopedEmmitter
@@ -26,7 +27,7 @@ public:
 		{
 			if (!Term->IsStructContextType() && (Term->Type.PinSubCategory != TEXT("self")))
 			{
-				SafetyConditions.Add(CppBackend.TermToText(EmitterContext, Term, false));
+				SafetyConditions.Add(CppBackend.TermToText(EmitterContext, Term, ENativizedTermUsage::Getter, false));
 			}
 		}
 
@@ -75,7 +76,7 @@ void FBlueprintCompilerCppBackend::EmitCallDelegateStatment(FEmitterLocalContext
 {
 	check(Statement.FunctionContext && Statement.FunctionContext->AssociatedVarProperty);
 	FSafeContextScopedEmmitter SafeContextScope(EmitterContext, Statement.FunctionContext->Context, *this);
-	EmitterContext.AddLine(FString::Printf(TEXT("%s.Broadcast(%s);"), *TermToText(EmitterContext, Statement.FunctionContext, false), *EmitMethodInputParameterList(EmitterContext, Statement)));
+	EmitterContext.AddLine(FString::Printf(TEXT("%s.Broadcast(%s);"), *TermToText(EmitterContext, Statement.FunctionContext, ENativizedTermUsage::Getter, false), *EmitMethodInputParameterList(EmitterContext, Statement)));
 }
 
 void FBlueprintCompilerCppBackend::EmitCallStatment(FEmitterLocalContext& EmitterContext, FKismetFunctionContext& FunctionContext, FBlueprintCompiledStatement& Statement)
@@ -101,25 +102,65 @@ void FBlueprintCompilerCppBackend::EmitCallStatment(FEmitterLocalContext& Emitte
 	}
 }
 
+struct FSetterExpressionBuilder
+{
+	FString EndCustomSetExpression;
+	FString DestinationExpression;
+
+	FSetterExpressionBuilder(FBlueprintCompilerCppBackend& CppBackend, FEmitterLocalContext& EmitterContext, FBPTerminal* LHS)
+	{
+		DestinationExpression = CppBackend.TermToText(EmitterContext, LHS, ENativizedTermUsage::Setter, false, &EndCustomSetExpression);
+	}
+
+	FString BuildStart() const 
+	{
+		FString Result = DestinationExpression;
+		const bool bCustomExpression = !EndCustomSetExpression.IsEmpty();
+		if (!bCustomExpression)
+		{
+			//It is not use regular = operator
+			Result += TEXT(" = ");
+		}
+		return Result;
+	}
+
+	FString BuildEnd(bool bAddSemicolon) const
+	{
+		FString Result = EndCustomSetExpression;
+		if (bAddSemicolon)
+		{
+			Result += TEXT(";");
+		}
+		return Result;
+	}
+
+	FString BuildFull(const FString& SourceExpression) const
+	{
+		const FString Start = BuildStart();
+		const FString End = BuildEnd(true);
+		return FString::Printf(TEXT("%s%s%s"), *Start, *SourceExpression, *End);
+	}
+};
+
 void FBlueprintCompilerCppBackend::EmitAssignmentStatment(FEmitterLocalContext& EmitterContext, FKismetFunctionContext& FunctionContext, FBlueprintCompiledStatement& Statement)
 {
 	check(Statement.LHS && Statement.RHS[0]);
-	const FString DestinationExpression = TermToText(EmitterContext, Statement.LHS, false, false);
-	const FString SourceExpression = TermToText(EmitterContext, Statement.RHS[0]);
-
+	const FString SourceExpression = TermToText(EmitterContext, Statement.RHS[0], ENativizedTermUsage::Getter);
+	FSetterExpressionBuilder SetterExpression(*this, EmitterContext, Statement.LHS);
 	FSafeContextScopedEmmitter SafeContextScope(EmitterContext, Statement.LHS->Context, *this);
 
 	FString BeginCast;
 	FString EndCast;
 	FEmitHelper::GenerateAutomaticCast(EmitterContext, Statement.LHS->Type, Statement.RHS[0]->Type, BeginCast, EndCast);
-	EmitterContext.AddLine(FString::Printf(TEXT("%s = %s%s%s;"), *DestinationExpression, *BeginCast, *SourceExpression, *EndCast));
+	const FString RHS = FString::Printf(TEXT("%s%s%s"), *BeginCast, *SourceExpression, *EndCast);
+	EmitterContext.AddLine(SetterExpression.BuildFull(RHS));
 }
 
 void FBlueprintCompilerCppBackend::EmitCastObjToInterfaceStatement(FEmitterLocalContext& EmitterContext, FKismetFunctionContext& FunctionContext, FBlueprintCompiledStatement& Statement)
 {
-	FString InterfaceClass = TermToText(EmitterContext, Statement.RHS[0]);
-	FString ObjectValue = TermToText(EmitterContext, Statement.RHS[1]);
-	FString InterfaceValue = TermToText(EmitterContext, Statement.LHS);
+	FString InterfaceClass = TermToText(EmitterContext, Statement.RHS[0], ENativizedTermUsage::Getter);
+	FString ObjectValue = TermToText(EmitterContext, Statement.RHS[1], ENativizedTermUsage::Getter);
+	FString InterfaceValue = TermToText(EmitterContext, Statement.LHS, ENativizedTermUsage::UnspecifiedOrReference);
 
 	// Both here and in UObject::execObjectToInterface IsValid function should be used.
 	EmitterContext.AddLine(FString::Printf(TEXT("if ( %s && %s->GetClass()->ImplementsInterface(%s) )"), *ObjectValue, *ObjectValue, *InterfaceClass));
@@ -136,9 +177,9 @@ void FBlueprintCompilerCppBackend::EmitCastObjToInterfaceStatement(FEmitterLocal
 
 void FBlueprintCompilerCppBackend::EmitCastBetweenInterfacesStatement(FEmitterLocalContext& EmitterContext, FKismetFunctionContext& FunctionContext, FBlueprintCompiledStatement& Statement)
 {
-	FString ClassToCastTo = TermToText(EmitterContext, Statement.RHS[0]);
-	FString InputInterface = TermToText(EmitterContext, Statement.RHS[1]);
-	FString ResultInterface = TermToText(EmitterContext, Statement.LHS);
+	FString ClassToCastTo = TermToText(EmitterContext, Statement.RHS[0], ENativizedTermUsage::Getter);
+	FString InputInterface = TermToText(EmitterContext, Statement.RHS[1], ENativizedTermUsage::Getter);
+	FString ResultInterface = TermToText(EmitterContext, Statement.LHS, ENativizedTermUsage::UnspecifiedOrReference);
 
 	FString InputObject = FString::Printf(TEXT("%s.GetObjectRef()"), *InputInterface);
 
@@ -154,7 +195,7 @@ void FBlueprintCompilerCppBackend::EmitCastBetweenInterfacesStatement(FEmitterLo
 	EmitterContext.AddLine(FString::Printf(TEXT("}")));
 }
 
-static void EmitGenericCast(FEmitterLocalContext& EmitterContext, UClass* ClassPtr, const FString& ObjectValue, const FString& ResultObject)
+static FString GenerateCastRHS(FEmitterLocalContext& EmitterContext, UClass* ClassPtr, const FString& ObjectValue)
 {
 	check(ClassPtr != nullptr);
 
@@ -163,50 +204,54 @@ static void EmitGenericCast(FEmitterLocalContext& EmitterContext, UClass* ClassP
 	{
 		const FString NativeClass = FEmitHelper::GetCppName(EmitterContext.GetFirstNativeOrConvertedClass(ClassPtr));
 		const FString TargetClass = EmitterContext.FindGloballyMappedObject(ClassPtr, UClass::StaticClass(), true);
-		EmitterContext.AddLine(FString::Printf(TEXT("%s = NoNativeCast<%s>(%s, %s);"), *ResultObject, *NativeClass, *TargetClass, *ObjectValue));
+		return FString::Printf(TEXT("NoNativeCast<%s>(%s, %s)"), *NativeClass, *TargetClass, *ObjectValue);
 	}
 	else
 	{
 		const FString TargetClass = FEmitHelper::GetCppName(ClassPtr);
-		EmitterContext.AddLine(FString::Printf(TEXT("%s = Cast<%s>(%s);"), *ResultObject, *TargetClass, *ObjectValue));
+		return FString::Printf(TEXT("Cast<%s>(%s)"), *TargetClass, *ObjectValue);
 	}
 }
 
 void FBlueprintCompilerCppBackend::EmitCastInterfaceToObjStatement(FEmitterLocalContext& EmitterContext, FKismetFunctionContext& FunctionContext, FBlueprintCompiledStatement& Statement)
 {
-	const FString InputInterface = TermToText(EmitterContext, Statement.RHS[1]);
-	const FString ResultObject = TermToText(EmitterContext, Statement.LHS, true, false);
+	const FString InputInterface = TermToText(EmitterContext, Statement.RHS[1], ENativizedTermUsage::Getter);
+	FSetterExpressionBuilder SetterExpression(*this, EmitterContext, Statement.LHS);
+	FSafeContextScopedEmmitter SafeContextScope(EmitterContext, Statement.LHS->Context, *this);
 
-	EmitGenericCast(EmitterContext
-		, CastChecked<UClass>(Statement.RHS[0]->ObjectLiteral)
-		, FString::Printf(TEXT("%s.GetObjectRef()"), *InputInterface)
-		, ResultObject);
+	const FString CastRHS = GenerateCastRHS(EmitterContext, CastChecked<UClass>(Statement.RHS[0]->ObjectLiteral), FString::Printf(TEXT("%s.GetObjectRef()"), *InputInterface));
+	EmitterContext.AddLine(SetterExpression.BuildFull(CastRHS));
 }
 
 void FBlueprintCompilerCppBackend::EmitDynamicCastStatement(FEmitterLocalContext& EmitterContext, FKismetFunctionContext& FunctionContext, FBlueprintCompiledStatement& Statement)
 {
-	const FString ObjectValue = TermToText(EmitterContext, Statement.RHS[1]);
-	const FString ResultObject = TermToText(EmitterContext, Statement.LHS, true, false);
-
-	EmitGenericCast(EmitterContext
-		, CastChecked<UClass>(Statement.RHS[0]->ObjectLiteral)
-		, ObjectValue
-		, ResultObject);
+	const FString ObjectValue = TermToText(EmitterContext, Statement.RHS[1], ENativizedTermUsage::Getter);
+	FSetterExpressionBuilder SetterExpression(*this, EmitterContext, Statement.LHS);
+	FSafeContextScopedEmmitter SafeContextScope(EmitterContext, Statement.LHS->Context, *this);
+	
+	const FString CastRHS = GenerateCastRHS(EmitterContext, CastChecked<UClass>(Statement.RHS[0]->ObjectLiteral), ObjectValue);
+	EmitterContext.AddLine(SetterExpression.BuildFull(CastRHS));
 }
 
 void FBlueprintCompilerCppBackend::EmitMetaCastStatement(FEmitterLocalContext& EmitterContext, FKismetFunctionContext& FunctionContext, FBlueprintCompiledStatement& Statement)
 {
-	FString DesiredClass = TermToText(EmitterContext, Statement.RHS[0]);
-	FString SourceClass = TermToText(EmitterContext, Statement.RHS[1]);
-	FString Destination = TermToText(EmitterContext, Statement.LHS, true, false);
-	EmitterContext.AddLine(FString::Printf(TEXT("%s = DynamicMetaCast(%s, %s);"),*Destination, *DesiredClass, *SourceClass));
+	const FString DesiredClass = TermToText(EmitterContext, Statement.RHS[0], ENativizedTermUsage::Getter);
+	const FString SourceClass = TermToText(EmitterContext, Statement.RHS[1], ENativizedTermUsage::Getter);
+	FSetterExpressionBuilder SetterExpression(*this, EmitterContext, Statement.LHS);
+	FSafeContextScopedEmmitter SafeContextScope(EmitterContext, Statement.LHS->Context, *this);
+
+	const FString CastRHS = FString::Printf(TEXT("DynamicMetaCast(%s, %s);"), *DesiredClass, *SourceClass);
+	EmitterContext.AddLine(SetterExpression.BuildFull(CastRHS));
 }
 
 void FBlueprintCompilerCppBackend::EmitObjectToBoolStatement(FEmitterLocalContext& EmitterContext, FKismetFunctionContext& FunctionContext, FBlueprintCompiledStatement& Statement)
 {
-	FString ObjectTarget = TermToText(EmitterContext, Statement.RHS[0]);
-	FString DestinationExpression = TermToText(EmitterContext, Statement.LHS, true, false);
-	EmitterContext.AddLine(FString::Printf(TEXT("%s = (%s != nullptr);"), *DestinationExpression, *ObjectTarget));
+	const FString ObjectTarget = TermToText(EmitterContext, Statement.RHS[0], ENativizedTermUsage::Getter);
+	FSetterExpressionBuilder SetterExpression(*this, EmitterContext, Statement.LHS);
+	FSafeContextScopedEmmitter SafeContextScope(EmitterContext, Statement.LHS->Context, *this);
+
+	const FString RHS = FString::Printf(TEXT("(%s != nullptr);"), *ObjectTarget);
+	EmitterContext.AddLine(SetterExpression.BuildFull(RHS));
 }
 
 void FBlueprintCompilerCppBackend::EmitAddMulticastDelegateStatement(FEmitterLocalContext& EmitterContext, FKismetFunctionContext& FunctionContext, FBlueprintCompiledStatement& Statement)
@@ -214,8 +259,8 @@ void FBlueprintCompilerCppBackend::EmitAddMulticastDelegateStatement(FEmitterLoc
 	check(Statement.LHS && Statement.LHS->AssociatedVarProperty);
 	FSafeContextScopedEmmitter SafeContextScope(EmitterContext, Statement.LHS->Context, *this);
 
-	const FString Delegate = TermToText(EmitterContext, Statement.LHS, false);
-	const FString DelegateToAdd = TermToText(EmitterContext, Statement.RHS[0]);
+	const FString Delegate = TermToText(EmitterContext, Statement.LHS, ENativizedTermUsage::UnspecifiedOrReference, false);
+	const FString DelegateToAdd = TermToText(EmitterContext, Statement.RHS[0], ENativizedTermUsage::Getter);
 
 	EmitterContext.AddLine(FString::Printf(TEXT("%s.Add(%s);"), *Delegate, *DelegateToAdd));
 }
@@ -225,8 +270,8 @@ void FBlueprintCompilerCppBackend::EmitRemoveMulticastDelegateStatement(FEmitter
 	check(Statement.LHS && Statement.LHS->AssociatedVarProperty);
 	FSafeContextScopedEmmitter SafeContextScope(EmitterContext, Statement.LHS->Context, *this);
 
-	const FString Delegate = TermToText(EmitterContext, Statement.LHS, false);
-	const FString DelegateToAdd = TermToText(EmitterContext, Statement.RHS[0]);
+	const FString Delegate = TermToText(EmitterContext, Statement.LHS, ENativizedTermUsage::UnspecifiedOrReference, false);
+	const FString DelegateToAdd = TermToText(EmitterContext, Statement.RHS[0], ENativizedTermUsage::Getter);
 
 	EmitterContext.AddLine(FString::Printf(TEXT("%s.Remove(%s);"), *Delegate, *DelegateToAdd));
 }
@@ -237,9 +282,9 @@ void FBlueprintCompilerCppBackend::EmitBindDelegateStatement(FEmitterLocalContex
 	check(Statement.LHS);
 	FSafeContextScopedEmmitter SafeContextScope(EmitterContext, Statement.LHS->Context, *this);
 
-	const FString Delegate = TermToText(EmitterContext, Statement.LHS, false);
-	const FString NameTerm = TermToText(EmitterContext, Statement.RHS[0]);
-	const FString ObjectTerm = TermToText(EmitterContext, Statement.RHS[1]);
+	const FString Delegate = TermToText(EmitterContext, Statement.LHS, ENativizedTermUsage::UnspecifiedOrReference, false);
+	const FString NameTerm = TermToText(EmitterContext, Statement.RHS[0], ENativizedTermUsage::Getter);
+	const FString ObjectTerm = TermToText(EmitterContext, Statement.RHS[1], ENativizedTermUsage::Getter);
 
 	EmitterContext.AddLine(FString::Printf(TEXT("%s.BindUFunction(%s,%s);"), *Delegate, *ObjectTerm, *NameTerm));
 }
@@ -249,7 +294,7 @@ void FBlueprintCompilerCppBackend::EmitClearMulticastDelegateStatement(FEmitterL
 	check(Statement.LHS);
 	FSafeContextScopedEmmitter SafeContextScope(EmitterContext, Statement.LHS->Context, *this);
 
-	const FString Delegate = TermToText(EmitterContext, Statement.LHS, false);
+	const FString Delegate = TermToText(EmitterContext, Statement.LHS, ENativizedTermUsage::UnspecifiedOrReference, false);
 
 	EmitterContext.AddLine(FString::Printf(TEXT("%s.Clear();"), *Delegate));
 }
@@ -257,14 +302,14 @@ void FBlueprintCompilerCppBackend::EmitClearMulticastDelegateStatement(FEmitterL
 void FBlueprintCompilerCppBackend::EmitCreateArrayStatement(FEmitterLocalContext& EmitterContext, FKismetFunctionContext& FunctionContext, FBlueprintCompiledStatement& Statement)
 {
 	FBPTerminal* ArrayTerm = Statement.LHS;
-	const FString Array = TermToText(EmitterContext, ArrayTerm);
+	const FString Array = TermToText(EmitterContext, ArrayTerm, ENativizedTermUsage::UnspecifiedOrReference);
 
 	EmitterContext.AddLine(FString::Printf(TEXT("%s.SetNum(%d, true);"), *Array, Statement.RHS.Num()));
 
 	for (int32 i = 0; i < Statement.RHS.Num(); ++i)
 	{
 		FBPTerminal* CurrentTerminal = Statement.RHS[i];
-		EmitterContext.AddLine(FString::Printf(TEXT("%s[%d] = %s;"), *Array, i, *TermToText(EmitterContext, CurrentTerminal)));
+		EmitterContext.AddLine(FString::Printf(TEXT("%s[%d] = %s;"), *Array, i, *TermToText(EmitterContext, CurrentTerminal, ENativizedTermUsage::Getter)));
 	}
 }
 
@@ -275,7 +320,7 @@ void FBlueprintCompilerCppBackend::EmitGotoStatement(FEmitterLocalContext& Emitt
 		if (bUseGotoState)
 		{
 			FString NextStateExpression;
-			NextStateExpression = TermToText(EmitterContext, Statement.LHS);
+			NextStateExpression = TermToText(EmitterContext, Statement.LHS, ENativizedTermUsage::Getter);
 
 			EmitterContext.AddLine(FString::Printf(TEXT("CurrentState = %s;"), *NextStateExpression));
 			EmitterContext.AddLine(FString::Printf(TEXT("break;\n")));
@@ -284,7 +329,7 @@ void FBlueprintCompilerCppBackend::EmitGotoStatement(FEmitterLocalContext& Emitt
 	else if ((Statement.Type == KCST_GotoIfNot) || (Statement.Type == KCST_EndOfThreadIfNot) || (Statement.Type == KCST_GotoReturnIfNot))
 	{
 		FString ConditionExpression;
-		ConditionExpression = TermToText(EmitterContext, Statement.LHS);
+		ConditionExpression = TermToText(EmitterContext, Statement.LHS, ENativizedTermUsage::Getter);
 
 		EmitterContext.AddLine(FString::Printf(TEXT("if (!%s)"), *ConditionExpression));
 		EmitterContext.AddLine(FString::Printf(TEXT("{")));
@@ -407,15 +452,15 @@ FString FBlueprintCompilerCppBackend::EmitSwitchValueStatmentInner(FEmitterLocal
 	FString Result = FString::Printf(TEXT("TSwitchValue<%s, %s>(%s, %s, %d")
 		, *IndexDeclaration
 		, *ValueDeclaration
-		, *TermToText(EmitterContext, IndexTerm) //index
-		, *TermToText(EmitterContext, DefaultValueTerm) // default
+		, *TermToText(EmitterContext, IndexTerm, ENativizedTermUsage::UnspecifiedOrReference) //index
+		, *TermToText(EmitterContext, DefaultValueTerm, ENativizedTermUsage::UnspecifiedOrReference) // default
 		, NumCases);
 
 	for (int32 TermIndex = TermsBeforeCases; TermIndex < (NumCases * TermsPerCase); TermIndex += TermsPerCase)
 	{
 		auto TermToRef = [&](const FBPTerminal* Term) -> FString
 		{
-			const FString TermEvaluation = TermToText(EmitterContext, Term); //should bGetter be false ?
+			const FString TermEvaluation = TermToText(EmitterContext, Term, ENativizedTermUsage::UnspecifiedOrReference); //should bGetter be false ?
 			if (Term->bIsLiteral) //TODO it should be done for every term, that cannot be handled as reference.
 			{
 				const FString LocalVarName = EmitterContext.GenerateUniqueLocalName();
@@ -424,7 +469,7 @@ FString FBlueprintCompilerCppBackend::EmitSwitchValueStatmentInner(FEmitterLocal
 			}
 			return TermEvaluation;
 		};
-		const FString Term0_Index = TermToText(EmitterContext, Statement.RHS[TermIndex]);
+		const FString Term0_Index = TermToText(EmitterContext, Statement.RHS[TermIndex], ENativizedTermUsage::UnspecifiedOrReference);
 		const FString Term1_Ref = TermToRef(Statement.RHS[TermIndex + 1]);
 		Result += FString::Printf(TEXT(", TSwitchPair<%s, %s>(%s, %s)")
 			, *IndexDeclaration
@@ -482,12 +527,14 @@ FString FBlueprintCompilerCppBackend::EmitMethodInputParameterList(FEmitterLocal
 				FEdGraphPinType LType;
 				auto Schema = GetDefault<UEdGraphSchema_K2>();
 				check(Schema);
+				ENativizedTermUsage TermUsage = ENativizedTermUsage::UnspecifiedOrReference;
 				if (Schema->ConvertPropertyToPinType(FuncParamProperty, LType))
 				{
 					FEmitHelper::GenerateAutomaticCast(EmitterContext, LType, Term->Type, BeginCast, CloseCast);
+					TermUsage = LType.bIsReference ? ENativizedTermUsage::UnspecifiedOrReference : ENativizedTermUsage::Getter;
 				}
 				VarName += BeginCast;
-				VarName += TermToText(EmitterContext, Term);
+				VarName += TermToText(EmitterContext, Term, TermUsage);
 				VarName += CloseCast;
 			}
 
@@ -527,12 +574,12 @@ static FString CustomThunkFunctionPostfix(FBlueprintCompiledStatement& Statement
 		}
 	}
 
-	FString FunctionName = Statement.FunctionToCall->GetName();
+	const FName FunctionName = Statement.FunctionToCall->GetFName();
 	if (ArrayTerm 
-		&& (FunctionName == TEXT("Array_Find") 
-			|| FunctionName == TEXT("Array_Contains") 
-			|| FunctionName == TEXT("Array_RemoveItem")
-			|| FunctionName == TEXT("Array_AddUnique")
+		&& ((FunctionName == GET_FUNCTION_NAME_CHECKED(UKismetArrayLibrary, Array_Find))
+			|| (FunctionName == GET_FUNCTION_NAME_CHECKED(UKismetArrayLibrary, Array_Contains))
+			|| (FunctionName == GET_FUNCTION_NAME_CHECKED(UKismetArrayLibrary, Array_RemoveItem))
+			|| (FunctionName == GET_FUNCTION_NAME_CHECKED(UKismetArrayLibrary, Array_AddUnique))
 			))
 	{
 		if (UEdGraphSchema_K2::PC_Text == ArrayTerm->Type.PinCategory)
@@ -576,12 +623,16 @@ FString FBlueprintCompilerCppBackend::EmitCallStatmentInner(FEmitterLocalContext
 
 	FString Result;
 	FString CloseCast;
+	TAutoPtr<FSetterExpressionBuilder> SetterExpression;
 	if (!bInline)
 	{
 		// Handle the return value of the function being called
 		UProperty* FuncToCallReturnProperty = Statement.FunctionToCall->GetReturnProperty();
 		if (FuncToCallReturnProperty && ensure(Statement.LHS))
 		{
+			SetterExpression = new FSetterExpressionBuilder(*this, EmitterContext, Statement.LHS);
+			Result += SetterExpression->BuildStart();
+
 			FString BeginCast;
 			FEdGraphPinType RType;
 			auto Schema = GetDefault<UEdGraphSchema_K2>();
@@ -590,7 +641,7 @@ FString FBlueprintCompilerCppBackend::EmitCallStatmentInner(FEmitterLocalContext
 			{
 				FEmitHelper::GenerateAutomaticCast(EmitterContext, Statement.LHS->Type, RType, BeginCast, CloseCast);
 			}
-			Result += FString::Printf(TEXT("%s = %s"), *TermToText(EmitterContext, Statement.LHS), *BeginCast);
+			Result += BeginCast;
 		}
 	}
 
@@ -605,7 +656,7 @@ FString FBlueprintCompilerCppBackend::EmitCallStatmentInner(FEmitterLocalContext
 			const FBPTerminal* const ValueTerminal = Statement.RHS.IsValidIndex(1) ? Statement.RHS[1] : nullptr;
 			if (Enum && ValueTerminal)
 			{
-				const FString Value = TermToText(EmitterContext, ValueTerminal);
+				const FString Value = TermToText(EmitterContext, ValueTerminal, ENativizedTermUsage::Getter);
 				const FString EnumName = FEmitHelper::GetCppName(Enum);
 				Result += FString::Printf(TEXT("%s__GetUserFriendlyName(EnumToByte<%s>(TEnumAsByte<%s>(%s)))"), *EnumName, *EnumName, *EnumName, *Value);
 				return true;
@@ -625,7 +676,7 @@ FString FBlueprintCompilerCppBackend::EmitCallStatmentInner(FEmitterLocalContext
 			Result += FString::Printf(TEXT("%s::Execute_%s(%s.GetObject() ")
 				, *FEmitHelper::GetCppName(ContextInterfaceClass)
 				, *FunctionToCallOriginalName
-				, *TermToText(EmitterContext, Statement.FunctionContext, false));
+				, *TermToText(EmitterContext, Statement.FunctionContext, ENativizedTermUsage::Getter, false));
 		}
 		else
 		{
@@ -641,7 +692,7 @@ FString FBlueprintCompilerCppBackend::EmitCallStatmentInner(FEmitterLocalContext
 				ensure(!bStaticCall); //unsupported yet
 				ensure(bCallOnDifferentObject); //unexpected
 				const FString WrapperName = FString::Printf(TEXT("FUnconvertedWrapper__%s"), *FEmitHelper::GetCppName(OwnerBPGC));
-				const FString CalledObject = bCallOnDifferentObject ? TermToText(EmitterContext, Statement.FunctionContext, false) : TEXT("this");
+				const FString CalledObject = bCallOnDifferentObject ? TermToText(EmitterContext, Statement.FunctionContext, ENativizedTermUsage::UnspecifiedOrReference, false) : TEXT("this");
 				Result += FString::Printf(TEXT("%s(%s)."), *WrapperName, *CalledObject);
 			}
 			else if (bStaticCall)
@@ -651,7 +702,7 @@ FString FBlueprintCompilerCppBackend::EmitCallStatmentInner(FEmitterLocalContext
 			}
 			else if (bCallOnDifferentObject) //@TODO: Badness, could be a self reference wired to another instance!
 			{
-				Result += FString::Printf(TEXT("%s->"), *TermToText(EmitterContext, Statement.FunctionContext, false));
+				Result += FString::Printf(TEXT("%s->"), *TermToText(EmitterContext, Statement.FunctionContext, ENativizedTermUsage::Getter, false));
 			}
 
 			if (Statement.bIsParentContext)
@@ -684,6 +735,10 @@ FString FBlueprintCompilerCppBackend::EmitCallStatmentInner(FEmitterLocalContext
 	}
 
 	Result += CloseCast;
+	if (SetterExpression.IsValid())
+	{
+		Result += SetterExpression->BuildEnd(false);
+	}
 	if (!bInline)
 	{
 		Result += TEXT(";");
@@ -697,15 +752,23 @@ FString FBlueprintCompilerCppBackend::EmitArrayGetByRef(FEmitterLocalContext& Em
 	check(Statement.RHS.Num() == 2);
 
 	FString Result;
-	Result += TermToText(EmitterContext, Statement.RHS[0]);
+	Result += TermToText(EmitterContext, Statement.RHS[0], ENativizedTermUsage::UnspecifiedOrReference);
 	Result += TEXT("[");
-	Result += TermToText(EmitterContext, Statement.RHS[1]);
+	Result += TermToText(EmitterContext, Statement.RHS[1], ENativizedTermUsage::Getter);
 	Result += TEXT("]");
 	return Result;
 }
 
-FString FBlueprintCompilerCppBackend::TermToText(FEmitterLocalContext& EmitterContext, const FBPTerminal* Term, bool bUseSafeContext, bool bGetter)
+FString FBlueprintCompilerCppBackend::TermToText(FEmitterLocalContext& EmitterContext, const FBPTerminal* Term, const ENativizedTermUsage TermUsage, const bool bUseSafeContext, FString* EndCustomSetExpression)
 {
+	ensure((TermUsage != ENativizedTermUsage::Setter) || !bUseSafeContext);
+	ensure((TermUsage == ENativizedTermUsage::Setter) == (EndCustomSetExpression != nullptr));
+	if (EndCustomSetExpression)
+	{
+		EndCustomSetExpression->Reset();
+	}
+
+	const bool bGetter = (TermUsage == ENativizedTermUsage::Getter);
 	const FString PSC_Self(TEXT("self"));
 	if (Term->bIsLiteral)
 	{
@@ -755,7 +818,7 @@ FString FBlueprintCompilerCppBackend::TermToText(FEmitterLocalContext& EmitterCo
 				}
 			}
 
-			ContextStr += TermToText(EmitterContext, Term->Context, false);
+			ContextStr += TermToText(EmitterContext, Term->Context, ENativizedTermUsage::UnspecifiedOrReference, false); // Should we just pass TermUsage?
 			ContextStr += bFromDefaultValue ? TEXT(")") : TEXT("");
 		}
 
@@ -768,7 +831,7 @@ FString FBlueprintCompilerCppBackend::TermToText(FEmitterLocalContext& EmitterCo
 			bIsAccessible &= !Term->AssociatedVarProperty->HasAnyPropertyFlags(CPF_NativeAccessSpecifierPrivate | CPF_NativeAccessSpecifierProtected);
 			if (!bIsAccessible)
 			{
-				ResultPath = FEmitHelper::AccessInaccessibleProperty(EmitterContext, Term->AssociatedVarProperty, ContextStr, TEXT("&"), 0, bGetter);
+				ResultPath = FEmitHelper::AccessInaccessibleProperty(EmitterContext, Term->AssociatedVarProperty, ContextStr, TEXT("&"), 0, TermUsage, EndCustomSetExpression);
 			}
 			else
 			{
@@ -804,7 +867,7 @@ FString FBlueprintCompilerCppBackend::TermToText(FEmitterLocalContext& EmitterCo
 					ensure(ContextStr.IsEmpty());
 					ContextStr = TEXT("this");
 				}
-				ResultPath = FEmitHelper::AccessInaccessibleProperty(EmitterContext, Term->AssociatedVarProperty, ContextStr, FString(), 0, bGetter);
+				ResultPath = FEmitHelper::AccessInaccessibleProperty(EmitterContext, Term->AssociatedVarProperty, ContextStr, FString(), 0, TermUsage, EndCustomSetExpression);
 			}
 			else
 			{
@@ -863,7 +926,11 @@ FString FBlueprintCompilerCppBackend::TermToText(FEmitterLocalContext& EmitterCo
 		{
 			const FString DefaultValueConstructor = FEmitHelper::DefaultValue(EmitterContext, Term->Type);
 			const FString DefaultValueVariable = EmitterContext.GenerateUniqueLocalName();
-			const FString CppType = FEmitHelper::PinTypeToNativeType(Term->Type);
+			const uint32 PropertyExportFlags = EPropertyExportCPPFlags::CPPF_CustomTypeName | EPropertyExportCPPFlags::CPPF_BlueprintCppBackend | EPropertyExportCPPFlags::CPPF_NoConst;
+			const FString CppType = Term->AssociatedVarProperty 
+				? EmitterContext.ExportCppDeclaration(Term->AssociatedVarProperty, EExportedDeclaration::Local, PropertyExportFlags, true)
+				: FEmitHelper::PinTypeToNativeType(Term->Type);
+
 			EmitterContext.AddLine(*FString::Printf(TEXT("%s %s = %s;"), *CppType, *DefaultValueVariable, *DefaultValueConstructor));
 			return FString::Printf(TEXT("((%s) ? (%s) : (%s))"), *Conditions, *ResultPath, *DefaultValueVariable);
 		}
diff --git a/Engine/Source/Developer/BlueprintCompilerCppBackend/Private/BlueprintCompilerCppBackend.h b/Engine/Source/Developer/BlueprintCompilerCppBackend/Private/BlueprintCompilerCppBackend.h
--- a/Engine/Source/Developer/BlueprintCompilerCppBackend/Private/BlueprintCompilerCppBackend.h
+++ b/Engine/Source/Developer/BlueprintCompilerCppBackend/Private/BlueprintCompilerCppBackend.h
@@ -5,6 +5,13 @@
 
 struct FEmitterLocalContext;
 
+enum class ENativizedTermUsage : uint8
+{
+	UnspecifiedOrReference,
+	Setter,
+	Getter,
+};
+
 /** Default implementation of BlueprintCompilerCppBackend. Code generated for function body usually contains a big "switch". */
 class FBlueprintCompilerCppBackend : public FBlueprintCompilerCppBackendBase
 {
@@ -57,5 +64,5 @@ protected:
 	FString EmitArrayGetByRef(FEmitterLocalContext& EmitterContext, FBlueprintCompiledStatement& Statement);
 
 public:
-	FString TermToText(FEmitterLocalContext& EmitterContext, const FBPTerminal* Term, bool bUseSafeContext = true, bool bGetter = true);
+	FString TermToText(FEmitterLocalContext& EmitterContext, const FBPTerminal* Term, const ENativizedTermUsage TermUsage, const bool bUseSafeContext = true, FString* EndCustomSetExpression = nullptr);
 };
\ No newline at end of file
diff --git a/Engine/Source/Developer/BlueprintCompilerCppBackend/Private/BlueprintCompilerCppBackendAnim.cpp b/Engine/Source/Developer/BlueprintCompilerCppBackend/Private/BlueprintCompilerCppBackendAnim.cpp
--- a/Engine/Source/Developer/BlueprintCompilerCppBackend/Private/BlueprintCompilerCppBackendAnim.cpp
+++ b/Engine/Source/Developer/BlueprintCompilerCppBackend/Private/BlueprintCompilerCppBackendAnim.cpp
@@ -33,7 +33,7 @@ void FBackendHelperAnim::CreateAnimClassData(FEmitterLocalContext& Context)
 				, FEmitDefaultValueHelper::EPropertyAccessOperator::Pointer);
 		}
 
-		Context.AddLine(FString::Printf(TEXT("InDynamicClass->AnimClassImplementation = %s;"), *LocalNativeName));
+		Context.AddLine(FString::Printf(TEXT("InDynamicClass->%s = %s;"), GET_MEMBER_NAME_STRING_CHECKED(UDynamicClass, AnimClassImplementation), *LocalNativeName));
 	}
 }
 
diff --git a/Engine/Source/Developer/BlueprintCompilerCppBackend/Private/BlueprintCompilerCppBackendBase.cpp b/Engine/Source/Developer/BlueprintCompilerCppBackend/Private/BlueprintCompilerCppBackendBase.cpp
--- a/Engine/Source/Developer/BlueprintCompilerCppBackend/Private/BlueprintCompilerCppBackendBase.cpp
+++ b/Engine/Source/Developer/BlueprintCompilerCppBackend/Private/BlueprintCompilerCppBackendBase.cpp
@@ -743,7 +743,10 @@ FString FBlueprintCompilerCppBackendBase::GenerateCodeFromEnum(UUserDefinedEnum*
 		const FString ElemName = EnumItemName(Index);
 		FString DisplayNameStr;
 		FTextStringHelper::WriteToString(DisplayNameStr, SourceEnum->GetEnumText(Index));
-		Header.AddLine(FString::Printf(TEXT("case %s::%s: FTextStringHelper::ReadFromString(TEXT(\"%s\"), Text); break;"), *EnumCppName, *ElemName, *DisplayNameStr.ReplaceCharWithEscapedChar()));
+		Header.AddLine(FString::Printf(TEXT("case %s::%s: FTextStringHelper::%s(TEXT(\"%s\"), Text); break;")
+			, *EnumCppName, *ElemName
+			, GET_FUNCTION_NAME_STRING_CHECKED(FTextStringHelper, ReadFromString)
+			, *DisplayNameStr.ReplaceCharWithEscapedChar()));
 	}
 
 	Header.AddLine(TEXT("default: ensure(false);"));
@@ -779,7 +782,7 @@ FString FBlueprintCompilerCppBackendBase::GenerateCodeFromStruct(UUserDefinedStr
 	EmitterContext.Header.AddLine(FString::Printf(TEXT("bool operator== (const %s& __Other) const"), *CppStructName));
 	EmitterContext.Header.AddLine(TEXT("{"));
 	EmitterContext.Header.IncreaseIndent();
-	EmitterContext.Header.AddLine(FString::Printf(TEXT("return %s::StaticStruct()->CompareScriptStruct(this, &__Other, 0);"), *CppStructName));
+	EmitterContext.Header.AddLine(FString::Printf(TEXT("return %s::StaticStruct()->%s(this, &__Other, 0);"), *CppStructName, GET_FUNCTION_NAME_STRING_CHECKED(UScriptStruct, CompareScriptStruct)));
 	EmitterContext.Header.DecreaseIndent();
 	EmitterContext.Header.AddLine(TEXT("};"));
 
@@ -941,7 +944,7 @@ FString FBlueprintCompilerCppBackendBase::GenerateWrapperForClass(UClass* Source
 			EmitterContext.Header.AddLine(TEXT("if (nullptr == __Property)"));
 			EmitterContext.Header.AddLine(TEXT("{"));
 			EmitterContext.Header.IncreaseIndent();
-			EmitterContext.Header.AddLine(FString::Printf(TEXT("__Property = GetClass()->FindPropertyByName(FName(TEXT(\"%s\")));"), *Property->GetName()));
+			EmitterContext.Header.AddLine(FString::Printf(TEXT("__Property = GetClass()->%s(FName(TEXT(\"%s\")));"), GET_FUNCTION_NAME_STRING_CHECKED(UClass, FindPropertyByName), *Property->GetName()));
 			EmitterContext.Header.AddLine(TEXT("check(__Property);"));
 			EmitterContext.Header.AddLine(TEXT("__PropertyPtr = __Property;"));
 			EmitterContext.Header.DecreaseIndent();
@@ -949,7 +952,7 @@ FString FBlueprintCompilerCppBackendBase::GenerateWrapperForClass(UClass* Source
 		}
 		else
 		{
-			EmitterContext.Header.AddLine(FString::Printf(TEXT("const UProperty* __Property = GetClass()->FindPropertyByName(FName(TEXT(\"%s\")));"), *Property->GetName()));
+			EmitterContext.Header.AddLine(FString::Printf(TEXT("const UProperty* __Property = GetClass()->%s(FName(TEXT(\"%s\")));"), GET_FUNCTION_NAME_STRING_CHECKED(UClass, FindPropertyByName), *Property->GetName()));
 		}
 		EmitterContext.Header.AddLine(FString::Printf(TEXT("return *(__Property->ContainerPtrToValuePtr<%s>(__Object));"), *TypeDeclaration));
 		EmitterContext.Header.DecreaseIndent();
@@ -1021,7 +1024,7 @@ FString FBlueprintCompilerCppBackendBase::GenerateWrapperForClass(UClass* Source
 			EmitterContext.Header.AddLine(FString::Printf(TEXT("static const FName __FunctionName(TEXT(\"%s\"));"), *Func->GetName()));
 		}
 		const FString FuncNameStr = bUseStaticVariables ? TEXT("__FunctionName") : FString::Printf(TEXT("FName(TEXT(\"%s\"))"), *Func->GetName());
-		EmitterContext.Header.AddLine(FString::Printf(TEXT("UFunction* __Function = __Object->FindFunctionChecked(%s);"), *FuncNameStr));
+		EmitterContext.Header.AddLine(FString::Printf(TEXT("UFunction* __Function = __Object->%s(%s);"), GET_FUNCTION_NAME_STRING_CHECKED(UObject, FindFunctionChecked), *FuncNameStr));
 
 		const FString FuncParametersStructName = FuncParameters.Num() ? (FuncCppName + TEXT("_Parameters")) : FString();
 		if (FuncParameters.Num())
@@ -1038,7 +1041,7 @@ FString FBlueprintCompilerCppBackendBase::GenerateWrapperForClass(UClass* Source
 
 			EmitterContext.Header.AddLine(FString::Printf(TEXT("%s __Parameters { %s };"), *FuncParametersStructName, *RawParameterList));
 		}
-		EmitterContext.Header.AddLine(FString::Printf(TEXT("__Object->ProcessEvent(__Function, %s);"), FuncParameters.Num() ? TEXT("&__Parameters") : TEXT("nullptr")));
+		EmitterContext.Header.AddLine(FString::Printf(TEXT("__Object->%s(__Function, %s);"), GET_FUNCTION_NAME_STRING_CHECKED(UObject, ProcessEvent), FuncParameters.Num() ? TEXT("&__Parameters") : TEXT("nullptr")));
 		EmitterContext.Header.DecreaseIndent();
 		EmitterContext.Header.AddLine(TEXT("}"));
 	}
diff --git a/Engine/Source/Developer/BlueprintCompilerCppBackend/Private/BlueprintCompilerCppBackendGatherDependencies.cpp b/Engine/Source/Developer/BlueprintCompilerCppBackend/Private/BlueprintCompilerCppBackendGatherDependencies.cpp
--- a/Engine/Source/Developer/BlueprintCompilerCppBackend/Private/BlueprintCompilerCppBackendGatherDependencies.cpp
+++ b/Engine/Source/Developer/BlueprintCompilerCppBackend/Private/BlueprintCompilerCppBackendGatherDependencies.cpp
@@ -172,6 +172,24 @@ struct FFindHeadersToInclude : public FGatherConvertedClassDependenciesHelperBas
 				}
 			}
 		}
+
+		// Include classes of native subobjects
+		if (BPGC)
+		{
+			UClass* NativeSuperClass = BPGC->GetSuperClass();
+			for (; NativeSuperClass && !NativeSuperClass->HasAnyClassFlags(CLASS_Native); NativeSuperClass = NativeSuperClass->GetSuperClass())
+			{}
+			UObject* NativeCDO = NativeSuperClass ? NativeSuperClass->GetDefaultObject(false) : nullptr;
+			if (NativeCDO)
+			{
+				TArray<UObject*> DefaultSubobjects;
+				NativeCDO->GetDefaultSubobjects(DefaultSubobjects);
+				for (UObject* DefaultSubobject : DefaultSubobjects)
+				{
+					IncludeTheHeaderInBody(DefaultSubobject ? DefaultSubobject->GetClass() : nullptr);
+				}
+			}
+		}
 	}
 
 	virtual void HandleObjectReference(UObject*& InObject, const UObject* InReferencingObject, const UProperty* InReferencingProperty) override
diff --git a/Engine/Source/Developer/BlueprintCompilerCppBackend/Private/BlueprintCompilerCppBackendUMG.cpp b/Engine/Source/Developer/BlueprintCompilerCppBackend/Private/BlueprintCompilerCppBackendUMG.cpp
--- a/Engine/Source/Developer/BlueprintCompilerCppBackend/Private/BlueprintCompilerCppBackendUMG.cpp
+++ b/Engine/Source/Developer/BlueprintCompilerCppBackend/Private/BlueprintCompilerCppBackendUMG.cpp
@@ -10,9 +10,9 @@ void FBackendHelperUMG::WidgetFunctionsInHeader(FEmitterLocalContext& Context)
 {
 	if (Cast<UWidgetBlueprintGeneratedClass>(Context.GetCurrentlyGeneratedClass()))
 	{
-		Context.Header.AddLine(TEXT("virtual void GetSlotNames(TArray<FName>& SlotNames) const override;"));
-		Context.Header.AddLine(TEXT("virtual void PreSave(const class ITargetPlatform* TargetPlatform) override;"));
-		Context.Header.AddLine(TEXT("virtual void CustomNativeInitilize() override;"));
+		Context.Header.AddLine(FString::Printf(TEXT("virtual void %s(TArray<FName>& SlotNames) const override;"), GET_FUNCTION_NAME_STRING_CHECKED(UUserWidget, GetSlotNames)));
+		Context.Header.AddLine(FString::Printf(TEXT("virtual void %s(const class ITargetPlatform* TargetPlatform) override;"), GET_FUNCTION_NAME_STRING_CHECKED(UUserWidget, PreSave)));
+		Context.Header.AddLine(FString::Printf(TEXT("virtual void %s() override;"), GET_FUNCTION_NAME_STRING_CHECKED(UUserWidget, CustomNativeInitilize)));
 	}
 }
 
@@ -63,7 +63,7 @@ void FBackendHelperUMG::EmitWidgetInitializationFunctions(FEmitterLocalContext&
 		};
 
 		{	// GetSlotNames
-			Context.AddLine(FString::Printf(TEXT("void %s::GetSlotNames(TArray<FName>& SlotNames) const"), *CppClassName));
+			Context.AddLine(FString::Printf(TEXT("void %s::%s(TArray<FName>& SlotNames) const"), *CppClassName, GET_FUNCTION_NAME_STRING_CHECKED(UUserWidget, GetSlotNames)));
 			Context.AddLine(TEXT("{"));
 			Context.IncreaseIndent();
 
@@ -75,7 +75,7 @@ void FBackendHelperUMG::EmitWidgetInitializationFunctions(FEmitterLocalContext&
 		}
 
 		{	// CustomNativeInitilize
-			Context.AddLine(FString::Printf(TEXT("void %s::CustomNativeInitilize()"), *CppClassName));
+			Context.AddLine(FString::Printf(TEXT("void %s::%s()"), *CppClassName, GET_FUNCTION_NAME_STRING_CHECKED(UUserWidget, CustomNativeInitilize)));
 			Context.AddLine(TEXT("{"));
 			Context.IncreaseIndent();
 
@@ -84,7 +84,8 @@ void FBackendHelperUMG::EmitWidgetInitializationFunctions(FEmitterLocalContext&
 			const FString AnimationsArrayNativeName = GenerateLocalProperty(Context, FindFieldChecked<UArrayProperty>(UWidgetBlueprintGeneratedClass::StaticClass(), TEXT("Animations")), reinterpret_cast<const uint8*>(&WidgetClass->Animations));
 			const FString BindingsArrayNativeName = GenerateLocalProperty(Context, FindFieldChecked<UArrayProperty>(UWidgetBlueprintGeneratedClass::StaticClass(), TEXT("Bindings")), reinterpret_cast<const uint8*>(&WidgetClass->Bindings));
 					
-			Context.AddLine(FString::Printf(TEXT("UWidgetBlueprintGeneratedClass::InitializeWidgetStatic(this, GetClass(), %s, %s, %s, %s, %s, nullptr);")
+			Context.AddLine(FString::Printf(TEXT("UWidgetBlueprintGeneratedClass::%s(this, GetClass(), %s, %s, %s, %s, %s, nullptr);")
+				, GET_FUNCTION_NAME_STRING_CHECKED(UWidgetBlueprintGeneratedClass, InitializeWidgetStatic)
 				, WidgetClass->bCanEverTick ? TEXT("true") : TEXT("false")
 				, WidgetClass->bCanEverPaint ? TEXT("true") : TEXT("false")
 				, *WidgetTreeStr
@@ -96,13 +97,13 @@ void FBackendHelperUMG::EmitWidgetInitializationFunctions(FEmitterLocalContext&
 		}
 
 		// PreSave
-		Context.AddLine(FString::Printf(TEXT("void %s::PreSave(const class ITargetPlatform* TargetPlatform)"), *CppClassName));
+		Context.AddLine(FString::Printf(TEXT("void %s::%s(const class ITargetPlatform* TargetPlatform)"), *CppClassName, GET_FUNCTION_NAME_STRING_CHECKED(UUserWidget, PreSave)));
 		Context.AddLine(TEXT("{"));
 		Context.IncreaseIndent();
-		Context.AddLine(TEXT("Super::PreSave(TargetPlatform);"));
+		Context.AddLine(FString::Printf(TEXT("Super::%s(TargetPlatform);"), GET_FUNCTION_NAME_STRING_CHECKED(UObject, PreSave)));
 		Context.AddLine(TEXT("TArray<FName> LocalNamedSlots;"));
-		Context.AddLine(TEXT("GetSlotNames(LocalNamedSlots);"));
-		Context.AddLine(TEXT("RemoveObsoleteBindings(LocalNamedSlots);"));
+		Context.AddLine(FString::Printf(TEXT("%s(LocalNamedSlots);"), GET_FUNCTION_NAME_STRING_CHECKED(UUserWidget, GetSlotNames)));
+		Context.AddLine(TEXT("RemoveObsoleteBindings(LocalNamedSlots);")); //RemoveObsoleteBindings is protected - no check
 		Context.DecreaseIndent();
 		Context.AddLine(TEXT("}"));
 	}
diff --git a/Engine/Source/Developer/BlueprintCompilerCppBackend/Private/BlueprintCompilerCppBackendUtils.cpp b/Engine/Source/Developer/BlueprintCompilerCppBackend/Private/BlueprintCompilerCppBackendUtils.cpp
--- a/Engine/Source/Developer/BlueprintCompilerCppBackend/Private/BlueprintCompilerCppBackendUtils.cpp
+++ b/Engine/Source/Developer/BlueprintCompilerCppBackend/Private/BlueprintCompilerCppBackendUtils.cpp
@@ -59,29 +59,37 @@ FString FEmitterLocalContext::FindGloballyMappedObject(const UObject* Object, co
 		int32 ObjectsCreatedPerClassIdx = MiscConvertedSubobjects.IndexOfByKey(Object);
 		if (INDEX_NONE != ObjectsCreatedPerClassIdx)
 		{
-			return FString::Printf(TEXT("CastChecked<%s>(CastChecked<UDynamicClass>(%s::StaticClass())->MiscConvertedSubobjects[%d])")
-				, *ClassString(), *FEmitHelper::GetCppName(ActualClass), ObjectsCreatedPerClassIdx);
+			return FString::Printf(TEXT("CastChecked<%s>(CastChecked<UDynamicClass>(%s::StaticClass())->%s[%d])")
+				, *ClassString(), *FEmitHelper::GetCppName(ActualClass)
+				, GET_MEMBER_NAME_STRING_CHECKED(UDynamicClass, MiscConvertedSubobjects)
+				, ObjectsCreatedPerClassIdx);
 		}
 
 		ObjectsCreatedPerClassIdx = ObjectsCreatedPerClassIdx = DynamicBindingObjects.IndexOfByKey(Object);
 		if (INDEX_NONE != ObjectsCreatedPerClassIdx)
 		{
-			return FString::Printf(TEXT("CastChecked<%s>(CastChecked<UDynamicClass>(%s::StaticClass())->DynamicBindingObjects[%d])")
-				, *ClassString(), *FEmitHelper::GetCppName(ActualClass), ObjectsCreatedPerClassIdx);
+			return FString::Printf(TEXT("CastChecked<%s>(CastChecked<UDynamicClass>(%s::StaticClass())->%s[%d])")
+				, *ClassString(), *FEmitHelper::GetCppName(ActualClass)
+				, GET_MEMBER_NAME_STRING_CHECKED(UDynamicClass, DynamicBindingObjects)
+				, ObjectsCreatedPerClassIdx);
 		}
 
 		ObjectsCreatedPerClassIdx = ComponentTemplates.IndexOfByKey(Object);
 		if (INDEX_NONE != ObjectsCreatedPerClassIdx)
 		{
-			return FString::Printf(TEXT("CastChecked<%s>(CastChecked<UDynamicClass>(%s::StaticClass())->ComponentTemplates[%d])")
-				, *ClassString(), *FEmitHelper::GetCppName(ActualClass), ObjectsCreatedPerClassIdx);
+			return FString::Printf(TEXT("CastChecked<%s>(CastChecked<UDynamicClass>(%s::StaticClass())->%s[%d])")
+				, *ClassString(), *FEmitHelper::GetCppName(ActualClass)
+				, GET_MEMBER_NAME_STRING_CHECKED(UDynamicClass, ComponentTemplates)
+				, ObjectsCreatedPerClassIdx);
 		}
 
 		ObjectsCreatedPerClassIdx = Timelines.IndexOfByKey(Object);
 		if (INDEX_NONE != ObjectsCreatedPerClassIdx)
 		{
-			return FString::Printf(TEXT("CastChecked<%s>(CastChecked<UDynamicClass>(%s::StaticClass())->Timelines[%d])")
-				, *ClassString(), *FEmitHelper::GetCppName(ActualClass), ObjectsCreatedPerClassIdx);
+			return FString::Printf(TEXT("CastChecked<%s>(CastChecked<UDynamicClass>(%s::StaticClass())->%s[%d])")
+				, *ClassString(), *FEmitHelper::GetCppName(ActualClass)
+				, GET_MEMBER_NAME_STRING_CHECKED(UDynamicClass, Timelines)
+				, ObjectsCreatedPerClassIdx);
 		}
 
 		if ((CurrentCodeType == EGeneratedCodeType::SubobjectsOfClass) || (CurrentCodeType == EGeneratedCodeType::CommonConstructor))
@@ -147,9 +155,10 @@ FString FEmitterLocalContext::FindGloballyMappedObject(const UObject* Object, co
 		const int32 EnumIndex = EnumsInCurrentClass.IndexOfByKey(UDE);
 		if (INDEX_NONE != EnumIndex)
 		{
-			return FString::Printf(TEXT("CastChecked<%s>(CastChecked<UDynamicClass>(%s::StaticClass())->ReferencedConvertedFields[%d])")
+			return FString::Printf(TEXT("CastChecked<%s>(CastChecked<UDynamicClass>(%s::StaticClass())->%s[%d])")
 				, *ClassString()
 				, *FEmitHelper::GetCppName(ActualClass)
+				, GET_MEMBER_NAME_STRING_CHECKED(UDynamicClass, ReferencedConvertedFields)
 				, EnumIndex);
 		}
 	}
@@ -168,9 +177,10 @@ FString FEmitterLocalContext::FindGloballyMappedObject(const UObject* Object, co
 
 			if (INDEX_NONE != AssetIndex)
 			{
-				return FString::Printf(TEXT("CastChecked<%s>(CastChecked<UDynamicClass>(%s::StaticClass())->UsedAssets[%d], ECastCheckedType::NullAllowed)")
+				return FString::Printf(TEXT("CastChecked<%s>(CastChecked<UDynamicClass>(%s::StaticClass())->%s[%d], ECastCheckedType::NullAllowed)")
 					, *ClassString()
 					, *FEmitHelper::GetCppName(ActualClass)
+					, GET_MEMBER_NAME_STRING_CHECKED(UDynamicClass, UsedAssets)
 					, AssetIndex);
 			}
 		}
@@ -822,10 +832,10 @@ void FEmitHelper::EmitLifetimeReplicatedPropsImpl(FEmitterLocalContext& EmitterC
 		{
 			if (!bFunctionInitilzed)
 			{
-				EmitterContext.AddLine(FString::Printf(TEXT("void %s::GetLifetimeReplicatedProps(TArray< FLifetimeProperty > & OutLifetimeProps) const"), *CppClassName));
+				EmitterContext.AddLine(FString::Printf(TEXT("void %s::%s(TArray< FLifetimeProperty > & OutLifetimeProps) const"), *CppClassName, GET_FUNCTION_NAME_STRING_CHECKED(UActorComponent, GetLifetimeReplicatedProps)));
 				EmitterContext.AddLine(TEXT("{"));
 				EmitterContext.IncreaseIndent();
-				EmitterContext.AddLine(TEXT("Super::GetLifetimeReplicatedProps(OutLifetimeProps);"));
+				EmitterContext.AddLine(FString::Printf(TEXT("Super::%s(OutLifetimeProps);"), GET_FUNCTION_NAME_STRING_CHECKED(UActorComponent, GetLifetimeReplicatedProps)));
 				bFunctionInitilzed = true;
 			}
 			EmitterContext.AddLine(FString::Printf(TEXT("DOREPLIFETIME_DIFFNAMES(%s, %s, FName(TEXT(\"%s\")));"), *CppClassName, *FEmitHelper::GetCppName(*It), *(It->GetName())));
@@ -1332,9 +1342,10 @@ FString FEmitHelper::GenerateGetPropertyByName(FEmitterLocalContext& EmitterCont
 		EmitterContext.IncreaseIndent();
 
 		const FString PropertyOwnerStruct = EmitterContext.FindGloballyMappedObject(Property->GetOwnerStruct(), UStruct::StaticClass());
-		EmitterContext.AddLine(FString::Printf(TEXT("%s = (%s)->FindPropertyByName(FName(TEXT(\"%s\")));")
+		EmitterContext.AddLine(FString::Printf(TEXT("%s = (%s)->%s(FName(TEXT(\"%s\")));")
 			, *PropertyPtrName
 			, *PropertyOwnerStruct
+			, GET_FUNCTION_NAME_STRING_CHECKED(UStruct, FindPropertyByName)
 			, *Property->GetName()));
 		EmitterContext.AddLine(FString::Printf(TEXT("check(%s);"), *PropertyPtrName));
 		EmitterContext.AddLine(FString::Printf(TEXT("%s = %s;"), *PropertyWeakPtrName, *PropertyPtrName));
@@ -1347,6 +1358,7 @@ FString FEmitHelper::GenerateGetPropertyByName(FEmitterLocalContext& EmitterCont
 		EmitterContext.AddLine(FString::Printf(TEXT("const UProperty* %s = (%s)->FindPropertyByName(FName(TEXT(\"%s\")));")
 			, *PropertyPtrName
 			, *PropertyOwnerStruct
+			, GET_FUNCTION_NAME_STRING_CHECKED(UStruct, FindPropertyByName)
 			, *Property->GetName()));
 		EmitterContext.AddLine(FString::Printf(TEXT("check(%s);"), *PropertyPtrName));
 	}
@@ -1359,22 +1371,43 @@ FString FEmitHelper::GenerateGetPropertyByName(FEmitterLocalContext& EmitterCont
 }
 
 FString FEmitHelper::AccessInaccessibleProperty(FEmitterLocalContext& EmitterContext, const UProperty* Property
-	, const FString& ContextStr, const FString& ContextAdressOp, int32 StaticArrayIdx, bool bGetter)
+	, const FString& ContextStr, const FString& ContextAdressOp, int32 StaticArrayIdx, ENativizedTermUsage TermUsage, FString* CustomSetExpressionEnding)
 {
 	check(Property);
+	ensure((TermUsage == ENativizedTermUsage::Setter) == (CustomSetExpressionEnding != nullptr));
+	if (CustomSetExpressionEnding)
+	{
+		CustomSetExpressionEnding->Reset();
+	}
+
 	const FString PropertyLocalName = GenerateGetPropertyByName(EmitterContext, Property);
 	const UBoolProperty* BoolProperty = Cast<const UBoolProperty>(Property);
-	if (bGetter && BoolProperty)
+	if ((TermUsage == ENativizedTermUsage::Getter) && BoolProperty)
 	{
-		return  FString::Printf(TEXT("(((UBoolProperty*)%s)->GetPropertyValue_InContainer(%s(%s), %d))")
+		return  FString::Printf(TEXT("(((UBoolProperty*)%s)->%s(%s(%s), %d))")
 			, *PropertyLocalName
+			, GET_FUNCTION_NAME_STRING_CHECKED(UBoolProperty, GetPropertyValue_InContainer)
 			, *ContextAdressOp
 			, *ContextStr
 			, StaticArrayIdx);
 	}
+
+	if ((TermUsage == ENativizedTermUsage::Setter) && BoolProperty)
+	{
+		if (ensure(CustomSetExpressionEnding))
+		{
+			*CustomSetExpressionEnding = FString::Printf(TEXT(", %d))"), StaticArrayIdx);
+		}
+		return FString::Printf(TEXT("(((UBoolProperty*)%s)->%s(%s(%s), ")
+			, *PropertyLocalName
+			, GET_FUNCTION_NAME_STRING_CHECKED(UBoolProperty, SetPropertyValue_InContainer)
+			, *ContextAdressOp
+			, *ContextStr);
+	}
+
 	if (BoolProperty && !BoolProperty->IsNativeBool())
 	{
-		UE_LOG(LogK2Compiler, Error, TEXT("AccessInaccessiblePropertyUsingOffset - bitfield %s"), *GetPathNameSafe(Property));
+		UE_LOG(LogK2Compiler, Error, TEXT("AccessInaccessibleProperty - bitfield %s"), *GetPathNameSafe(Property));
 	}
 	const uint32 CppTemplateTypeFlags = EPropertyExportCPPFlags::CPPF_CustomTypeName
 		| EPropertyExportCPPFlags::CPPF_NoConst | EPropertyExportCPPFlags::CPPF_NoRef | EPropertyExportCPPFlags::CPPF_NoStaticArray
@@ -1503,17 +1536,46 @@ FString FEmitHelper::AccessInaccessiblePropertyUsingOffset(FEmitterLocalContext&
 		, StaticArrayIdx);
 }
 
-bool FBackendHelperStaticSearchableValues::HasSearchableValues(UClass* Class)
+struct FSearchableValuesdHelper_StaticData
 {
-	TArray<FStringClassReference> ClassesWithSearchableValues;
-	//TODO: export to .ini
-	ClassesWithSearchableValues.Add(FStringClassReference(TEXT("/Script/GameplayAbilities.GameplayCueNotify_Actor")));
-	ClassesWithSearchableValues.Add(FStringClassReference(TEXT("/Script/GameplayAbilities.GameplayCueNotify_Static")));
+	TArray<FStringClassReference> ClassesWithStaticSearchableValues;
+	TArray<FName> TagPropertyNames;
+private:
+	FSearchableValuesdHelper_StaticData()
+	{
+		{
+			TArray<FString> Paths;
+			GConfig->GetArray(TEXT("BlueprintNativizationSettings"), TEXT("ClassesWithStaticSearchableValues"), Paths, GEditorIni);
+			for (FString& Path : Paths)
+			{
+				ClassesWithStaticSearchableValues.Add(FStringClassReference(Path));
+			}
+		}
+
+		{
+			TArray<FString> Names;
+			GConfig->GetArray(TEXT("BlueprintNativizationSettings"), TEXT("StaticSearchableTagNames"), Names, GEditorIni);
+			for (FString& Name : Names)
+			{
+				TagPropertyNames.Add(FName(*Name));
+			}
+		}
+	}
 
-	for (FStringClassReference& IterClassRef : ClassesWithSearchableValues)
+public:
+	static FSearchableValuesdHelper_StaticData& Get()
 	{
-		UClass* IterClass = IterClassRef.ResolveClass();
-		if (IterClass && Class && Class->IsChildOf(IterClass))
+		static FSearchableValuesdHelper_StaticData StaticInstance;
+		return StaticInstance;
+	}
+};
+
+bool FBackendHelperStaticSearchableValues::HasSearchableValues(UClass* InClass)
+{
+	for (const FStringClassReference& ClassStrRef : FSearchableValuesdHelper_StaticData::Get().ClassesWithStaticSearchableValues)
+	{
+		UClass* IterClass = ClassStrRef.ResolveClass();
+		if (IterClass && InClass && InClass->IsChildOf(IterClass))
 		{
 			return true;
 		}
@@ -1553,12 +1615,7 @@ void FBackendHelperStaticSearchableValues::EmitFunctionDefinition(FEmitterLocalC
 	if(ensure(OriginalSourceClass))
 	{
 		FAssetData ClassAsset(OriginalSourceClass);
-
-		TArray<FName> TagPropertyNames;
-		//TODO: export to ini
-		TagPropertyNames.Add(FName(TEXT("GameplayCueName")));
-
-		for (const FName TagPropertyName : TagPropertyNames)
+		for (const FName TagPropertyName : FSearchableValuesdHelper_StaticData::Get().TagPropertyNames)
 		{
 			const FName FoundValue = ClassAsset.GetTagValueRef<FName>(TagPropertyName);
 			if (!FoundValue.IsNone())
diff --git a/Engine/Source/Developer/BlueprintCompilerCppBackend/Private/BlueprintCompilerCppBackendUtils.h b/Engine/Source/Developer/BlueprintCompilerCppBackend/Private/BlueprintCompilerCppBackendUtils.h
--- a/Engine/Source/Developer/BlueprintCompilerCppBackend/Private/BlueprintCompilerCppBackendUtils.h
+++ b/Engine/Source/Developer/BlueprintCompilerCppBackend/Private/BlueprintCompilerCppBackendUtils.h
@@ -2,6 +2,7 @@
 #pragma  once
 
 #include "BlueprintCompilerCppBackendGatherDependencies.h"
+#include "BlueprintCompilerCppBackend.h"
 
 struct FCodeText
 {
@@ -89,17 +90,17 @@ public:
 	{
 		if (ListType == EClassSubobjectList::ComponentTemplates)
 		{
-			return TEXT("ComponentTemplates");
+			return GET_MEMBER_NAME_STRING_CHECKED(UDynamicClass, ComponentTemplates);
 		}
 		if (ListType == EClassSubobjectList::Timelines)
 		{
-			return TEXT("Timelines");
+			return GET_MEMBER_NAME_STRING_CHECKED(UDynamicClass, Timelines);
 		}
 		if (ListType == EClassSubobjectList::DynamicBindingObjects)
 		{
-			return TEXT("DynamicBindingObjects");
+			return GET_MEMBER_NAME_STRING_CHECKED(UDynamicClass, DynamicBindingObjects);
 		}
-		return TEXT("MiscConvertedSubobjects");
+		return GET_MEMBER_NAME_STRING_CHECKED(UDynamicClass, MiscConvertedSubobjects);
 	}
 
 	void RegisterClassSubobject(UObject* Object, EClassSubobjectList ListType)
@@ -260,7 +261,8 @@ struct FEmitHelper
 	static FString GenerateGetPropertyByName(FEmitterLocalContext& EmitterContext, const UProperty* Property);
 
 	static FString AccessInaccessibleProperty(FEmitterLocalContext& EmitterContext, const UProperty* Property
-		, const FString& ContextStr, const FString& ContextAdressOp, int32 StaticArrayIdx, bool bGetter);
+		, const FString& ContextStr, const FString& ContextAdressOp, int32 StaticArrayIdx
+		, ENativizedTermUsage TermUsage, FString* CustomSetExpressionEnding);
 
 	// This code works properly as long, as all fields in structures are UProperties!
 	static FString AccessInaccessiblePropertyUsingOffset(FEmitterLocalContext& EmitterContext, const UProperty* Property
diff --git a/Engine/Source/Developer/BlueprintCompilerCppBackend/Private/BlueprintCompilerCppBackendValueHelper.cpp b/Engine/Source/Developer/BlueprintCompilerCppBackend/Private/BlueprintCompilerCppBackendValueHelper.cpp
--- a/Engine/Source/Developer/BlueprintCompilerCppBackend/Private/BlueprintCompilerCppBackendValueHelper.cpp
+++ b/Engine/Source/Developer/BlueprintCompilerCppBackend/Private/BlueprintCompilerCppBackendValueHelper.cpp
@@ -79,8 +79,9 @@ void FEmitDefaultValueHelper::OuterGenerate(FEmitterLocalContext& Context
 				{
 					const FString PropertyLocalName = FEmitHelper::GenerateGetPropertyByName(Context, Property);
 					const FString ValueStr = Context.ExportTextItem(Property, Property->ContainerPtrToValuePtr<uint8>(DataContainer, ArrayIndex));
-					Context.AddLine(FString::Printf(TEXT("(((UBoolProperty*)%s)->SetPropertyValue_InContainer(%s(%s), %s, %d));")
+					Context.AddLine(FString::Printf(TEXT("(((UBoolProperty*)%s)->%s(%s(%s), %s, %d));")
 						, *PropertyLocalName
+						, GET_FUNCTION_NAME_STRING_CHECKED(UBoolProperty, SetPropertyValue_InContainer)
 						, *OperatorStr
 						, *ContainerStr
 						, *ValueStr
@@ -89,7 +90,7 @@ void FEmitDefaultValueHelper::OuterGenerate(FEmitterLocalContext& Context
 				}
 				const FString GetPtrStr = bNoExportProperty
 					? FEmitHelper::AccessInaccessiblePropertyUsingOffset(Context, Property, ContainerStr, OperatorStr, ArrayIndex)
-					: FEmitHelper::AccessInaccessibleProperty(Context, Property, ContainerStr, OperatorStr, ArrayIndex, false);
+					: FEmitHelper::AccessInaccessibleProperty(Context, Property, ContainerStr, OperatorStr, ArrayIndex, ENativizedTermUsage::UnspecifiedOrReference, nullptr);
 				PathToMember = Context.GenerateUniqueLocalName();
 				Context.AddLine(FString::Printf(TEXT("auto& %s = %s;"), *PathToMember, *GetPtrStr));
 			}
@@ -220,7 +221,11 @@ void FEmitDefaultValueHelper::InnerGenerate(FEmitterLocalContext& Context, const
 			if (RegularInnerStruct)
 			{
 				const FString InnerStructStr = Context.FindGloballyMappedObject(RegularInnerStruct, UScriptStruct::StaticClass());
-				Context.AddLine(FString::Printf(TEXT("%s->InitializeStruct(%s.GetData(), %d);"), *InnerStructStr, *PathToMember, ScriptArrayHelper.Num()));
+				Context.AddLine(FString::Printf(TEXT("%s->%s(%s.GetData(), %d);")
+					, *InnerStructStr
+					, GET_FUNCTION_NAME_STRING_CHECKED(UStruct, InitializeStruct)
+					, *PathToMember
+					, ScriptArrayHelper.Num()));
 			}
 		}
 
@@ -350,15 +355,15 @@ bool FEmitDefaultValueHelper::SpecialStructureConstructor(const UStruct* Struct,
 			const FFloatRangeBound* FloatRangeBound = reinterpret_cast<const FFloatRangeBound*>(ValuePtr);
 			if (FloatRangeBound->IsExclusive())
 			{
-				*OutResult = FString::Printf(TEXT("FFloatRangeBound::Exclusive(%s)"), *FEmitHelper::FloatToString(FloatRangeBound->GetValue()));
+				*OutResult = FString::Printf(TEXT("FFloatRangeBound::%s(%s)"), GET_FUNCTION_NAME_STRING_CHECKED(FFloatRangeBound, Exclusive), *FEmitHelper::FloatToString(FloatRangeBound->GetValue()));
 			}
 			if (FloatRangeBound->IsInclusive())
 			{
-				*OutResult = FString::Printf(TEXT("FFloatRangeBound::Inclusive(%s)"), *FEmitHelper::FloatToString(FloatRangeBound->GetValue()));
+				*OutResult = FString::Printf(TEXT("FFloatRangeBound::%s(%s)"), GET_FUNCTION_NAME_STRING_CHECKED(FFloatRangeBound, Inclusive), *FEmitHelper::FloatToString(FloatRangeBound->GetValue()));
 			}
 			if (FloatRangeBound->IsOpen())
 			{
-				*OutResult = TEXT("FFloatRangeBound::Open()");
+				*OutResult = FString::Printf(TEXT("FFloatRangeBound::%s()"), GET_FUNCTION_NAME_STRING_CHECKED(FFloatRangeBound, Open));
 			}
 		}
 		return true;
@@ -390,15 +395,15 @@ bool FEmitDefaultValueHelper::SpecialStructureConstructor(const UStruct* Struct,
 			const FInt32RangeBound* RangeBound = reinterpret_cast<const FInt32RangeBound*>(ValuePtr);
 			if (RangeBound->IsExclusive())
 			{
-				*OutResult = FString::Printf(TEXT("FInt32RangeBound::Exclusive(%d)"), RangeBound->GetValue());
+				*OutResult = FString::Printf(TEXT("FInt32RangeBound::%s(%d)"), GET_FUNCTION_NAME_STRING_CHECKED(FInt32RangeBound, Exclusive), RangeBound->GetValue());
 			}
 			if (RangeBound->IsInclusive())
 			{
-				*OutResult = FString::Printf(TEXT("FInt32RangeBound::Inclusive(%d)"), RangeBound->GetValue());
+				*OutResult = FString::Printf(TEXT("FInt32RangeBound::%s(%d)"), GET_FUNCTION_NAME_STRING_CHECKED(FInt32RangeBound, Exclusive), RangeBound->GetValue());
 			}
 			if (RangeBound->IsOpen())
 			{
-				*OutResult = TEXT("FInt32RangeBound::Open()");
+				*OutResult = FString::Printf(TEXT("FInt32RangeBound::%s()"), GET_FUNCTION_NAME_STRING_CHECKED(FFloatRangeBound, Open));
 			}
 		}
 		return true;
@@ -523,14 +528,11 @@ struct FNonativeComponentData
 
 	bool HandledAsSpecialProperty(FEmitterLocalContext& Context, const UProperty* Property)
 	{
-		static const FName RelativeLocationName(TEXT("RelativeLocation"));
-		static const FName RelativeRotationName(TEXT("RelativeRotation"));
-
 		// skip relative location and rotation. THey are ignored for root components created from scs (and they probably should be reset by scs editor).
 		if (bIsRoot && (Property->GetOuter() == USceneComponent::StaticClass()))
 		{
-			UProperty* RelativeLocationProperty = USceneComponent::StaticClass()->FindPropertyByName(RelativeLocationName);
-			UProperty* RelativeRotationProperty = USceneComponent::StaticClass()->FindPropertyByName(RelativeRotationName);
+			UProperty* RelativeLocationProperty = USceneComponent::StaticClass()->FindPropertyByName(GET_MEMBER_NAME_CHECKED(USceneComponent, RelativeLocation));
+			UProperty* RelativeRotationProperty = USceneComponent::StaticClass()->FindPropertyByName(GET_MEMBER_NAME_CHECKED(USceneComponent, RelativeRotation));
 			if ((Property == RelativeLocationProperty) || (Property == RelativeRotationProperty))
 			{
 				return true;
@@ -545,13 +547,16 @@ struct FNonativeComponentData
 		ensure(!NativeVariablePropertyName.IsEmpty());
 		if (bSetNativeCreationMethod)
 		{
-			Context.AddLine(FString::Printf(TEXT("%s->CreationMethod = EComponentCreationMethod::Native;"), *NativeVariablePropertyName));
+			Context.AddLine(FString::Printf(TEXT("%s->%s = EComponentCreationMethod::Native;"), *NativeVariablePropertyName, GET_MEMBER_NAME_STRING_CHECKED(UActorComponent, CreationMethod)));
 		}
 
 		if (!ParentVariableName.IsEmpty())
 		{
 			const FString SocketName = (AttachToName == NAME_None) ? FString() : FString::Printf(TEXT(", TEXT(\"%s\")"), *AttachToName.ToString());
-			Context.AddLine(FString::Printf(TEXT("%s->AttachToComponent(%s, FAttachmentTransformRules::KeepRelativeTransform %s);"), *NativeVariablePropertyName, *ParentVariableName, *SocketName));
+			Context.AddLine(FString::Printf(TEXT("%s->%s(%s, FAttachmentTransformRules::KeepRelativeTransform %s);")
+				, *NativeVariablePropertyName
+				, GET_FUNCTION_NAME_STRING_CHECKED(USceneComponent, AttachToComponent)
+				, *ParentVariableName, *SocketName));
 			// AttachTo is called first in case some properties will be overridden.
 		}
 
@@ -570,11 +575,11 @@ struct FNonativeComponentData
 
 	void EmitForcedPostLoad(FEmitterLocalContext& Context)
 	{
-		Context.AddLine(FString::Printf(TEXT("if(%s && !%s->IsTemplate())"), *NativeVariablePropertyName, *NativeVariablePropertyName));
+		Context.AddLine(FString::Printf(TEXT("if(%s && !%s->%s())"), *NativeVariablePropertyName, *NativeVariablePropertyName, GET_FUNCTION_NAME_STRING_CHECKED(UActorComponent, IsTemplate)));
 		Context.AddLine(TEXT("{"));
 		Context.IncreaseIndent();
-		Context.AddLine(FString::Printf(TEXT("%s->SetFlags(RF_NeedPostLoad |RF_NeedPostLoadSubobjects);"), *NativeVariablePropertyName));
-		Context.AddLine(FString::Printf(TEXT("%s->ConditionalPostLoad();"), *NativeVariablePropertyName));
+		Context.AddLine(FString::Printf(TEXT("%s->%s(RF_NeedPostLoad |RF_NeedPostLoadSubobjects);"), *NativeVariablePropertyName, GET_FUNCTION_NAME_STRING_CHECKED(UActorComponent, SetFlags)));
+		Context.AddLine(FString::Printf(TEXT("%s->%s();"), *NativeVariablePropertyName, GET_FUNCTION_NAME_STRING_CHECKED(UActorComponent, ConditionalPostLoad)));
 		Context.DecreaseIndent();
 		Context.AddLine(TEXT("}"));
 	}
@@ -804,12 +809,12 @@ void FEmitDefaultValueHelper::GenerateCustomDynamicClassInitializationUsedAssets
 	Context.AddLine(FString::Printf(TEXT("void %s::__CustomDynamicClassInitializationUsedAssets(UDynamicClass* InDynamicClass)"), *CppClassName));
 	Context.AddLine(TEXT("{"));
 	Context.IncreaseIndent();
-	Context.AddLine(TEXT("ensure(0 == InDynamicClass->UsedAssets.Num());"));
+	Context.AddLine(FString::Printf(TEXT("ensure(0 == InDynamicClass->%s.Num());"), GET_MEMBER_NAME_STRING_CHECKED(UDynamicClass, UsedAssets)));
 
 	for (auto LocAsset : Context.UsedObjectInCurrentClass)
 	{
 		const FString AssetStr = Context.FindGloballyMappedObject(LocAsset, UObject::StaticClass(), true, false);
-		Context.AddLine(FString::Printf(TEXT("InDynamicClass->UsedAssets.Add(%s);"), *AssetStr));
+		Context.AddLine(FString::Printf(TEXT("InDynamicClass->%s.Add(%s);"), GET_MEMBER_NAME_STRING_CHECKED(UDynamicClass, UsedAssets), *AssetStr));
 	}
 
 	Context.DecreaseIndent();
@@ -824,13 +829,13 @@ void FEmitDefaultValueHelper::GenerateCustomDynamicClassInitialization(FEmitterL
 	Context.AddLine(FString::Printf(TEXT("void %s::__CustomDynamicClassInitialization(UDynamicClass* InDynamicClass)"), *CppClassName));
 	Context.AddLine(TEXT("{"));
 	Context.IncreaseIndent();
-	Context.AddLine(TEXT("ensure(0 == InDynamicClass->ReferencedConvertedFields.Num());"));
-	Context.AddLine(TEXT("ensure(0 == InDynamicClass->MiscConvertedSubobjects.Num());"));
-	Context.AddLine(TEXT("ensure(0 == InDynamicClass->DynamicBindingObjects.Num());"));
-	Context.AddLine(TEXT("ensure(0 == InDynamicClass->ComponentTemplates.Num());"));
-	Context.AddLine(TEXT("ensure(0 == InDynamicClass->Timelines.Num());"));
-	Context.AddLine(TEXT("ensure(nullptr == InDynamicClass->AnimClassImplementation);"));
-	Context.AddLine(TEXT("InDynamicClass->AssembleReferenceTokenStream();"));
+	Context.AddLine(FString::Printf(TEXT("ensure(0 == InDynamicClass->%s.Num());"), GET_MEMBER_NAME_STRING_CHECKED(UDynamicClass, ReferencedConvertedFields)));
+	Context.AddLine(FString::Printf(TEXT("ensure(0 == InDynamicClass->%s.Num());"), GET_MEMBER_NAME_STRING_CHECKED(UDynamicClass, MiscConvertedSubobjects)));
+	Context.AddLine(FString::Printf(TEXT("ensure(0 == InDynamicClass->%s.Num());"), GET_MEMBER_NAME_STRING_CHECKED(UDynamicClass, DynamicBindingObjects)));
+	Context.AddLine(FString::Printf(TEXT("ensure(0 == InDynamicClass->%s.Num());"), GET_MEMBER_NAME_STRING_CHECKED(UDynamicClass, ComponentTemplates)));
+	Context.AddLine(FString::Printf(TEXT("ensure(0 == InDynamicClass->%s.Num());"), GET_MEMBER_NAME_STRING_CHECKED(UDynamicClass, Timelines)));
+	Context.AddLine(FString::Printf(TEXT("ensure(nullptr == InDynamicClass->%s);"), GET_MEMBER_NAME_STRING_CHECKED(UDynamicClass, AnimClassImplementation)));
+	Context.AddLine(FString::Printf(TEXT("InDynamicClass->%s();"), GET_FUNCTION_NAME_STRING_CHECKED(UDynamicClass, AssembleReferenceTokenStream)));
 
 	Context.CurrentCodeType = FEmitterLocalContext::EGeneratedCodeType::SubobjectsOfClass;
 	Context.ResetPropertiesForInaccessibleStructs();
@@ -841,7 +846,7 @@ void FEmitDefaultValueHelper::GenerateCustomDynamicClassInitialization(FEmitterL
 	}
 	for (auto LocEnum : Context.Dependencies.ConvertedEnum)
 	{
-		Context.AddLine(FString::Printf(TEXT("InDynamicClass->ReferencedConvertedFields.Add(LoadObject<UEnum>(nullptr, TEXT(\"%s\")));"), *(LocEnum->GetPathName().ReplaceCharWithEscapedChar())));
+		Context.AddLine(FString::Printf(TEXT("InDynamicClass->%s.Add(LoadObject<UEnum>(nullptr, TEXT(\"%s\")));"), GET_MEMBER_NAME_STRING_CHECKED(UDynamicClass, ReferencedConvertedFields), *(LocEnum->GetPathName().ReplaceCharWithEscapedChar())));
 		Context.EnumsInCurrentClass.Add(LocEnum);
 	}
 
@@ -861,7 +866,7 @@ void FEmitDefaultValueHelper::GenerateCustomDynamicClassInitialization(FEmitterL
 
 			const FString ClassConstructor = FDependenciesHelper::GenerateZConstructor(ClassToLoad);
 			Context.AddLine(FString::Printf(TEXT("extern UClass* %s;"), *ClassConstructor));
-			Context.AddLine(FString::Printf(TEXT("InDynamicClass->ReferencedConvertedFields.Add(%s);"), *ClassConstructor));
+			Context.AddLine(FString::Printf(TEXT("InDynamicClass->%s.Add(%s);"), GET_MEMBER_NAME_STRING_CHECKED(UDynamicClass, ReferencedConvertedFields), *ClassConstructor));
 
 			//Context.AddLine(FString::Printf(TEXT("InDynamicClass->ReferencedConvertedFields.Add(LoadObject<UClass>(nullptr, TEXT(\"%s\")));")
 			//	, *(ClassToLoad->GetPathName().ReplaceCharWithEscapedChar())));
@@ -880,7 +885,7 @@ void FEmitDefaultValueHelper::GenerateCustomDynamicClassInitialization(FEmitterL
 		}
 		const FString StructConstructor = FDependenciesHelper::GenerateZConstructor(LocStruct);
 		Context.AddLine(FString::Printf(TEXT("extern UScriptStruct* %s;"), *StructConstructor));
-		Context.AddLine(FString::Printf(TEXT("InDynamicClass->ReferencedConvertedFields.Add(%s);"), *StructConstructor));
+		Context.AddLine(FString::Printf(TEXT("InDynamicClass->%s.Add(%s);"), GET_MEMBER_NAME_STRING_CHECKED(UDynamicClass, ReferencedConvertedFields), *StructConstructor));
 	}
 
 	TArray<UActorComponent*> ActorComponentTempatesOwnedByClass = BPGC->ComponentTemplates;
@@ -1085,10 +1090,14 @@ void FEmitDefaultValueHelper::GenerateConstructor(FEmitterLocalContext& Context)
 
 				if (Cast<UPrimitiveComponent>(ComponentToInit.ComponentTemplate))
 				{
-					Context.AddLine(FString::Printf(TEXT("if(!%s->IsTemplate())"), *ComponentToInit.NativeVariablePropertyName));
+					Context.AddLine(FString::Printf(TEXT("if(!%s->%s())"), *ComponentToInit.NativeVariablePropertyName, GET_FUNCTION_NAME_STRING_CHECKED(UPrimitiveComponent, IsTemplate)));
 					Context.AddLine(TEXT("{"));
 					Context.IncreaseIndent();
-					Context.AddLine(FString::Printf(TEXT("%s->BodyInstance.FixupData(%s);"), *ComponentToInit.NativeVariablePropertyName, *ComponentToInit.NativeVariablePropertyName));
+					Context.AddLine(FString::Printf(TEXT("%s->%s.%s(%s);")
+						, *ComponentToInit.NativeVariablePropertyName
+						, GET_MEMBER_NAME_STRING_CHECKED(UPrimitiveComponent, BodyInstance)
+						, GET_FUNCTION_NAME_STRING_CHECKED(FBodyInstance, FixupData)
+						, *ComponentToInit.NativeVariablePropertyName));
 					Context.DecreaseIndent();
 					Context.AddLine(TEXT("}"));
 				}
@@ -1113,16 +1122,16 @@ void FEmitDefaultValueHelper::GenerateConstructor(FEmitterLocalContext& Context)
 	Context.ResetPropertiesForInaccessibleStructs();
 
 	Context.ResetPropertiesForInaccessibleStructs();
-	Context.AddLine(FString::Printf(TEXT("void %s::PostLoadSubobjects(FObjectInstancingGraph* OuterInstanceGraph)"), *CppClassName));
+	Context.AddLine(FString::Printf(TEXT("void %s::%s(FObjectInstancingGraph* OuterInstanceGraph)"), *CppClassName, GET_FUNCTION_NAME_STRING_CHECKED(UObject, PostLoadSubobjects)));
 	Context.AddLine(TEXT("{"));
 	Context.IncreaseIndent();
-	Context.AddLine(TEXT("Super::PostLoadSubobjects(OuterInstanceGraph);"));
+	Context.AddLine(FString::Printf(TEXT("Super::%s(OuterInstanceGraph);"), GET_FUNCTION_NAME_STRING_CHECKED(UObject, PostLoadSubobjects)));
 	for (auto& ComponentToFix : NativeCreatedComponentProperties)
 	{
 		Context.AddLine(FString::Printf(TEXT("if(%s)"), *ComponentToFix));
 		Context.AddLine(TEXT("{"));
 		Context.IncreaseIndent();
-		Context.AddLine(FString::Printf(TEXT("%s->CreationMethod = EComponentCreationMethod::Native;"), *ComponentToFix));
+		Context.AddLine(FString::Printf(TEXT("%s->%s = EComponentCreationMethod::Native;"), *ComponentToFix, GET_MEMBER_NAME_STRING_CHECKED(UActorComponent, CreationMethod)));
 		Context.DecreaseIndent();
 		Context.AddLine(TEXT("}"));
 	}
@@ -1249,8 +1258,11 @@ FString FEmitDefaultValueHelper::HandleInstancedSubobject(FEmitterLocalContext&
 		}
 		else
 		{
-			Context.AddLine(FString::Printf(TEXT("auto %s = CastChecked<%s>(GetDefaultSubobjectByName(TEXT(\"%s\")));")
-				, *LocalNativeName, *FEmitHelper::GetCppName(ObjectClass), *Object->GetName()));
+			Context.AddLine(FString::Printf(TEXT("auto %s = CastChecked<%s>(%s(TEXT(\"%s\")));")
+				, *LocalNativeName
+				, *FEmitHelper::GetCppName(ObjectClass)
+				, GET_FUNCTION_NAME_STRING_CHECKED(UObject, GetDefaultSubobjectByName)
+				, *Object->GetName()));
 		}
 
 		const UObject* ObjectArchetype = Object->GetArchetype();
diff --git a/Engine/Source/Developer/BlueprintNativeCodeGen/Private/BlueprintNativeCodeGenModule.cpp b/Engine/Source/Developer/BlueprintNativeCodeGen/Private/BlueprintNativeCodeGenModule.cpp
--- a/Engine/Source/Developer/BlueprintNativeCodeGen/Private/BlueprintNativeCodeGenModule.cpp
+++ b/Engine/Source/Developer/BlueprintNativeCodeGen/Private/BlueprintNativeCodeGenModule.cpp
@@ -134,21 +134,21 @@ namespace
 	{
 		TArray<UObject*> Objects;
 		GetObjectsWithOuter(Package, Objects, false);
-		for (auto Entry : Objects)
+		for (UObject* Entry : Objects)
 		{
 			if (Entry->HasAnyFlags(ExcludedFlags))
 			{
 				continue;
 			}
 
+			if (FBlueprintSupport::IsDeferredDependencyPlaceholder(Entry))
+			{
+				continue;
+			}
+
+			// Not a skeleton class
 			if (UClass* AsClass = Cast<UClass>(Entry))
 			{
-				// Not a placeholder
-				if (AsClass->GetName().Contains(TEXT("PLACEHOLDER")))
-				{
-					continue;
-				}
-				// Not a skeleton class
 				if (UBlueprint* GeneratingBP = Cast<UBlueprint>(AsClass->ClassGeneratedBy))
 				{
 					if (AsClass != GeneratingBP->GeneratedClass)
@@ -157,6 +157,7 @@ namespace
 					}
 				}
 			}
+
 			OutStruct = Cast<UStruct>(Entry);
 			if (OutStruct)
 			{
@@ -387,11 +388,14 @@ void FBlueprintNativeCodeGenModule::GenerateSingleAsset(UField* ForConversion, c
 		ConversionRecord.GeneratedHeaderPath.Empty();
 	}
 
-	check(bSuccess);
-	if (bSuccess)
+	if (ensure(bSuccess))
 	{
 		GetManifest(PlatformName).GatherModuleDependencies(ForConversion->GetOutermost());
 	}
+	else
+	{
+		UE_LOG(LogBlueprintCodeGen, Error, TEXT("FBlueprintNativeCodeGenModule::GenerateSingleAsset error: %s"), *GetPathNameSafe(ForConversion));
+	}
 
 	BackendPCHQuery.Unbind();
 }
diff --git a/Engine/Source/Developer/BlueprintProfiler/Private/BlueprintProfiler.cpp b/Engine/Source/Developer/BlueprintProfiler/Private/BlueprintProfiler.cpp
--- a/Engine/Source/Developer/BlueprintProfiler/Private/BlueprintProfiler.cpp
+++ b/Engine/Source/Developer/BlueprintProfiler/Private/BlueprintProfiler.cpp
@@ -86,6 +86,11 @@ void FBlueprintProfiler::InstrumentEvent(const EScriptInstrumentationEvent& Even
 			ScriptCodeOffset = StackFrame.Code - StackFrame.Node->Script.GetData() - 1;
 			CurrentFunctionName = StackFrame.Node->GetFName();
 		}
+		// Handle latent suspends - set the script offset to the LinkId of the most recent latent action.
+		if (Event.GetType() == EScriptInstrumentation::ResumeEvent)
+		{
+			ScriptCodeOffset = Event.GetLatentLinkId();
+		}
 		InstrumentationEventQueue.Add(FScriptInstrumentedEvent(Event.GetType(), CurrentFunctionName, ScriptCodeOffset));
 		// Reset context on event end
 		if (Event.GetType() == EScriptInstrumentation::Stop)
@@ -232,7 +237,6 @@ void FBlueprintProfiler::ProcessEventProfilingData()
 			}
 			case EScriptInstrumentation::Stop:
 			{
-				const int32 NumEventRanges = ScriptEventRanges.Num();
 				if (ScriptEventRanges.Num() > 0)
 				{
 					FEventRange& EventRangeToProcess = ScriptEventRanges.Last();
@@ -244,10 +248,14 @@ void FBlueprintProfiler::ProcessEventProfilingData()
 						// Find any resumed contexts if required
 						if (InstrumentationEventQueue[EventRangeToProcess.StartIdx].IsResumeEvent())
 						{
-							if (TSharedPtr<FScriptEventPlayback>* SuspendedEvent = SuspendedEvents.Find(EventRangeToProcess.InstanceName))
+							if (TMap<int32, TSharedPtr<FScriptEventPlayback>>* SuspendedEventMap = SuspendedEvents.Find(EventRangeToProcess.InstanceName))
 							{
-								EventToProcess = *SuspendedEvent;
-								SuspendedEvents.Remove(EventRangeToProcess.InstanceName);
+								const int32 LatentLinkId = InstrumentationEventQueue[EventRangeToProcess.StartIdx].GetScriptCodeOffset();
+								if (TSharedPtr<FScriptEventPlayback>* SuspendedEvent = SuspendedEventMap->Find(LatentLinkId))
+								{
+									EventToProcess = *SuspendedEvent;
+									SuspendedEventMap->Remove(LatentLinkId);
+								}
 							}
 						}
 						else
@@ -263,7 +271,8 @@ void FBlueprintProfiler::ProcessEventProfilingData()
 							DirtyContexts.Add(EventRangeToProcess.BlueprintContext);
 							if (EventToProcess->IsSuspended())
 							{
-								SuspendedEvents.Add(EventRangeToProcess.InstanceName) = EventToProcess;
+								TMap<int32, TSharedPtr<FScriptEventPlayback>>& SuspendedEventMap = SuspendedEvents.FindOrAdd(EventRangeToProcess.InstanceName);
+								SuspendedEventMap.Add(EventToProcess->GetLatentLinkId()) = EventToProcess;
 							}
 						}
 					}
