Copying //UE4/Dev-Blueprints to //UE4/Dev-Main (Source: //UE4/Dev-Blueprints @ 3057841)

==========================
MAJOR FEATURES + CHANGES
==========================

Change 3043219 on 2016/07/08 by [REDACTED_PERSON]

	SKismetInspector was incoherently using it's property view's EditingEnabledDelegate. Now assigning it only in one place. This causes the kismet inspector in the diff tool to be read only as client code intends

Change 3044844 on 2016/07/11 by [REDACTED_PERSON]

	Fixed issue where older CustomEvents had incorrect UserDefinedPin directions associated with them

Change 3045508 on 2016/07/11 by [REDACTED_PERSON]

	Ensuring the order in which we reinstance Blueprint components is correct (SCS templates need to come before their instances).

Change 3045986 on 2016/07/12 by [REDACTED_PERSON]

	Reuse hardcoded paths. Reduce amount of generated strings.

Change 3046188 on 2016/07/12 by [REDACTED_PERSON]

	Reduce number of TArray::operator[] calls.

Change 3046412 on 2016/07/12 by [REDACTED_PERSON]

	Reduced size of FBlueprintDependencyData constructor (no inlined code). Reducednumber of TArray.Add calls.

Change 3046608 on 2016/07/12 by [REDACTED_PERSON]

	Fixed crash when breaking all links when pins are split pins with an array type

Change 3046755 on 2016/07/12 by [REDACTED_PERSON]

	This improves handling of inherited functions and events in the blueprint profiler and adds the concept of class scope to the instrumented signals.
	#Proj BlueprintProfiler, Kismet

	- Renamed EScriptInstrumentationEvent to EScriptInstrumentationSignal for improved clarity
	- Minor tidy up/refactor in EScriptInstrumentationSignal, moved some of the code to scripcore.cpp for visibility of other object types and exported it under COREUOBJECT_API.
	- Added class scope tracking to capture context and EScriptInstrumentationSignal so we can emit signals when the scope changes.
	- Modified LatentActionManager to implement new EScriptInstrumentationSignal interfaces.
	- Added code in the blueprint profiler to create sub class contexts when a class scope switch is emitted.
	- Added option in the blueprint profiler toolbar to hide.show inherited events
	- Removed the code for expand all option from the blueprint profiler toolbar and ExecutionStatDisplay
	- Modified the node colors to use preferences where applicable and added new code to discover graphnode flags before exec node creation
	- Added Scoped function name and event name calls to FBlueprintExecutionContext and FBlueprintFunctionContext
	- Modified the compiler module to rebuild inherited classes with instumentation when a child is built.

Change 3047035 on 2016/07/12 by [REDACTED_PERSON]

	This adds some functionality to correctly map nested tunnels and macro's in the blueprint compiler so that correct source nodes and tunnels can be identified from the blueprint debug data.
	#Proj BlueprintGraph, KismetCompiler, UnrealEd, Engine

	- Modified FEdGraphUtilities::MergeChildrenGraphsIn to optionally create boundary nodes when recursively merging graphs, this is only active during instrumented blueprint compilation and will be refactored later as a more general solution.
	- Modified UK2Node_TunnelBoundary to find source node guids for the current graph and cache them away for use in the VM backend
	- Modified UK2Node_TunnelBoundary to locate the source tunnel instance (Macro or Composite Node) and cache away for use in the VM backend.
	- Removed some potentially problematic checks added into the KismetCompiler before 4.13 branch to avoid fallout, these will be added back in with care later.
	- Added some extra data into the CompilerResults log to make the above work and not get stomped on by recursive macro node injection, again will be refactored out later.

Change 3047206 on 2016/07/12 by [REDACTED_PERSON]

	This updates the blueprint profiler to make use of the new boundary node code to get the nested tunnel stats working.
	#Proj BlueprintProfiler

	- This should be the final piece bar a few minor issues. Updated the macro/tunnel mapping code and the event processing to make use of the new tunnel data brought in with CL 3047035
	- Still a problem that needs to be looked at regarding tracepaths in the widgets and tunnels, it might actually be more general than that but I don't think it's a huge concern/will be difficult to fix

Change 3047468 on 2016/07/13 by [REDACTED_PERSON]

	[UE-32254] Add a normalized heatmap view mode to the Blueprint profiler.

	change summary:
	- added EBlueprintProfilerHeatLevelMetricsType and changed EBlueprintProfilerHeatMapDisplayMode to an enum class
	- renamed GetNodeHeat*() methods to GetExclusiveHeat*()
	- added UBlueprintProfilerSettings and relocated heat map display settings from IBlueprintProfilerInterface/FBlueprintProfiler
	- changed access to heat map display settings to route through UBlueprintProfilerSettings instead of FBlueprintProfiler
	- modified UEditorExperimentalSettings::PostEditChangeProperty() to remove calls to set static heat level threshold limits; this is being handled elsewhere
	- modified FBlueprintProfiler ctor to remove calls to set static heat level threshold limits; this is now being handled elsewhere
	- modified SBlueprintProfilerToolbar to include a toggle switch for heat level metrics type; user ability to customize threshold settings is now tied to the toggle switch
	- added IBlueprintProfilerInterface::GetProfilerDataForBlueprint() API and implementation in FBlueprintProfiler
	- modified FScriptPerfData to include fields for caching calculated heat level values (now returned by the Get*() accessor methods, rather than calculating on each access)
	- added FScriptHeatLevelMetrics struct
	- added FScriptPerfData::SetHeatLevels()
	- added FScriptExecutionNode::CalculateHeatLevelStats()
	- added FScriptExecutionBlueprint::UpdateHeatLevelMetrics(), now being called from FScriptExecutionBlueprint::RefreshStats() to update heat level thresholds based on user-selected type
	- now calling CalculateHeatLevelStats() on BP/instance exec nodes as part of FScriptExecutionBlueprint::RefreshStats() (i.e. heat levels are now precalculated and cached after updating stats)
	- modified SGraphNodeK2Base::GetNodeIndicatorOverlayColor() to calculate the appropriate heat levels for the local aggregated PerfNode based on the current toggle switch setting (relative or custom)

Change 3047513 on 2016/07/13 by [REDACTED_PERSON]

	[UE-32254] Fix code merge oversight causing compile error.

Change 3048638 on 2016/07/13 by [REDACTED_PERSON]

	[UE-33131] Fix crash when compiling a Blueprint with multiple event graphs with instrumentation enabled.

	change summary:
	- added FBlueprintEditorUtils::IsEventGraph()
	- modified FBlueprintExecutionContext::MapBlueprintExecution() to avoid stomping the ubergraph entry
	- modified FBlueprintExecutionContext::GetFunctionNameFromGraph() to map secondary event graphs to the ubergraph function name
	- modified FBlueprintExecutionContext::GetScopedFunctionNameFromGraph() to map secondary event graphs to the ubergraph function name

	notes:
	- will also address UE-33133

Change 3048796 on 2016/07/13 by [REDACTED_PERSON]

	FAnimNode_BlendListByBool constructor works the same way both  in runtime and in editor.

	It is necessary for BP->C++ converter. To reduce size of code construction, a diff between a "default" structure instance is used (so the default instance must be the same in runtime).

Change 3048797 on 2016/07/13 by [REDACTED_PERSON]

	Fixed issue where output nodes sometimes weren't being created when collapsing nodes to macros

Change 3048805 on 2016/07/13 by [REDACTED_PERSON]

	Reduce size of struct construction.
	While generating code, structures are diff with default instances.
	FLatentActionInfo and FBodyInstance are handled in a special way.

Change 3048988 on 2016/07/13 by [REDACTED_PERSON]

	Some fixes to the pure node stat code to start to adress asserts recently introduced and re-enable the profiler widgets for pure nodes.
	#Proj BlueprintProfiler, Kismet

	- Part 1 of these changes, this addresses standard pure node asserts in the profiler.
	- The second part which is underway will address problems inside macro and tunnel instances.

Change 3050027 on 2016/07/14 by [REDACTED_PERSON]

	This fixes an issue adding instrumentation to blueprints that don't implement events
	#Proj BlueprintProfiler

	- The profiler made an assumption that a blueprint has an ubergraph function, whicn apparently is incorrect.
	- Removed some stale ubergraph code and added some extra checks when looking up function names.

Change 3050029 on 2016/07/14 by [REDACTED_PERSON]

	This modifies the blueprint instrumented compilation chain so only the the blueprint you compile and all dependencies are instrumented and the profiler is notified rather than waiting for event data.
	#Proj BlueprintProfiler, Kismet, UnrelEd

	- This also improves the execution graph UI, notifying the user that no instances are available to display data from.

Change 3050275 on 2016/07/14 by [REDACTED_PERSON]

	Updating the blueprint profiler connection policy to handle scoped function names and fix the wire heatmaps.
	#Proj BlueprintProfiler

Change 3050513 on 2016/07/14 by [REDACTED_PERSON]

	BP Profiler: Distinguish (by name) the top level event and the nested event node in the tree-view

Change 3050624 on 2016/07/14 by [REDACTED_PERSON]

	This switches off some blueprint profiler debug defines checked in errantly in CL 3048988
	#Proj Kismet

Change 3052034 on 2016/07/15 by [REDACTED_PERSON]

	This fixes some incorrect wording on a UI notification message warning about profiling data not being available because no instances are placed in the current level
	#Proj Kismet

Change 3052037 on 2016/07/15 by [REDACTED_PERSON]

	To maintain functionality, and address fallout from CL 3041603 - removing a callback that was previously never used (not until CL 3041603). If a Blueprint graph is read-only, that is not enough to block us from editing properties when inside that graph (think interfaces, math experssion graphs, etc.).

Change 3052271 on 2016/07/15 by [REDACTED_PERSON]

	BP Profiler: Added persistence for display settings

Change 3052416 on 2016/07/15 by [REDACTED_PERSON]

	This change adds new functionality to map pure links through tunnel boundaries and fixes some asserts caused by missing profiler nodes.
	#Proj BlueprintProfiler, Kismet

	- Added map through tunnel functionality for pure nodes
	- Added some code to patch in pure exec nodes to the tunnel calling context, this is to support pure nodes wired to tunnel pure pins which are mapped by the tunnel but really in the calling function.
	- Changed the way simple tunnel entry nodes create child widgets to get rid of a visual error.
	- Changed the blueprint compiler to use the updated true source nodes when registering pure ranges for impure nodes.

Change 3052791 on 2016/07/16 by [REDACTED_PERSON]

	Fix for the internal function appearance for profiler stats, it was failing to display function inner stats.
	#Proj BlueprintProfiler

Change 3052794 on 2016/07/16 by [REDACTED_PERSON]

	This fixes function return stats to be colored correctly in the blueprint profiler
	#Proj BlueprintProfiler

Change 3052957 on 2016/07/17 by [REDACTED_PERSON]

	Removing Sample Bias settings and the ability to disable pure node times in the blueprint profiler
	#Proj BlueprintProfiler, Kismet, UnrealEd

Change 3052964 on 2016/07/17 by [REDACTED_PERSON]

	Renaming exclusive timings internally inside the blueprint profiler.
	#Proj GraphEditor, Kismet, UnrealEd, BlueprintProfiler

	- Doesn't resolve UE-32264, its more of a supporting change
	- Refactoring node exclusive API naming conventions to reflect the decision to go with average.

Change 3052979 on 2016/07/17 by [REDACTED_PERSON]

	This fixes problems with sequence and tunnel instances with sequences inside them showing extra samples in the blueprint profiler
	#Proj BlueprintProfiler, Kismet

	- Added SampleBase so the perfdata can account for re-entrant samples.
	- Fixed sequence node re-entrant/sample count stats
	- Added a second number format for times and made times work and display in Millisecs rather than seconds.
	- Added trace histories for the tunnel entry/exit sites.

Change 3052986 on 2016/07/17 by [REDACTED_PERSON]

	This updates support for heatwires in the blueprint profiler connection drawing policy to support changes in the tunnel mapping.
	#Jia UE-30591 - Profiler Heatwires fail at times around macro nodes
	#Proj BlueprintProfiler

Change 3053027 on 2016/07/17 by [REDACTED_PERSON]

	Fix for some missing pure node entries that were causing and assert in certain blueprints when profiling.
	#Proj BlueprintProfiler, Kismet

	- Fixes to the way we gather tunnel instance pure nodes and the way we patch them into calling function contexts.

Change 3053886 on 2016/07/18 by [REDACTED_PERSON]

	Class- and asset-types are listed in UEdGraphSchema_K2::GetVariableTypeTree.

Change 3053904 on 2016/07/18 by [REDACTED_PERSON]

	This fixes an issue found when mapping nested tunnels in the blueprint profiler that caused them to display incorrectly.
	#Proj BlueprintProfiler, Kismet

	- A minor update to the exit site mapping of tunnels to avoid the issue.

Change 3053920 on 2016/07/18 by [REDACTED_PERSON]

	Disabling the heat display for the max time column in the blueprint profiler.
	#Proj Kismet

Change 3053985 on 2016/07/18 by [REDACTED_PERSON]

	BP Nativization:
	- Dependencies for UserDefinedStruct default values
	- Fixed array of class declaration.

Change 3054266 on 2016/07/18 by [REDACTED_PERSON]

	Missed something in CL 3053904 that caused the tunnel exit mapping to cause asserts. This just corrects the oversight ( bad check )
	#Proj BlueprintProfiler

Change 3054357 on 2016/07/18 by [REDACTED_PERSON]

	BP Profiler - misc. UI polish

	change summary:
	- stat columns are now right-justified; stats centered vertically within cells; samples column is centered
	- darkened the profiler view's content area a bit to help hotspot stats stand out a bit better
	- made it so that the columns will stretch to fill the content area when the BP editor window is maximized
	- exec node name column text now remains white when selected (to be consistent w/ stat text)
	- desaturated the tree view highlight to work better w/ the darker content area and light stat/name text

Change 3054504 on 2016/07/18 by [REDACTED_PERSON]

	Removing the hottest endpoint profiler wire option from the blueprint profiler.
	#Proj BlueprintProfiler, Kismet

Change 3055145 on 2016/07/19 by [REDACTED_PERSON]

	Fixed a missing header problem when "bDontNativizeDataOnlyBP=false"

Change 3055187 on 2016/07/19 by [REDACTED_PERSON]

Change 3055366 on 2016/07/19 by [REDACTED_PERSON]

	Fixed FBlueprintEditorUtils::GetDependentBlueprints. BP is not dependent on itself.

Change 3055415 on 2016/07/19 by [REDACTED_PERSON]

	Fixing up some average timing namings that regressed/were missed to exclusive.
	#Proj BlueprintProfiler, GraphEditor, Kismet

Change 3055708 on 2016/07/19 by [REDACTED_PERSON]

	[UE-32929] BP profiler tool - display/threshold options are now being saved.

	change summary:
	- added properties for stat display options to BP profiler settings.
	- moved custom threshold values from experimental editor settings to BP profiler settings object.
	- removed BP profiler settings accessor methods in favor of a direct get/set (to be consistent w/ other usages).
	- modified FBlueprintProfilerStatOptions ctor to sync display option flags w/ BP profiler settings (from config).
	- modified FBlueprintProfilerStatOptions::OnChecked() to sync BP profiler settings w/ display option flags (to config).
	- consolidated custom heat threshold value get/set UI delegate methods + added internal ECustomPerformanceThreshold type.
	- added global defaults for custom heat thresholds to BP profiler settings; these are currently being used to determine whether or not to show the "reset to default" button in the UI (previously was mirroring hard-coded values in the delegate methods).

Change 3055801 on 2016/07/19 by [REDACTED_PERSON]

	Reverting a SButton change from WEX merge, as it caused noticable problems when repeatedly clicking certain buttons in the editor.

Change 3056088 on 2016/07/19 by [REDACTED_PERSON]

	This is a general update on the stat calculations for blueprint to bring them into line with recent changes. ( and fix them )
	#Proj BlueprintProfiler, Kismet

	- Did a pass through tunnel stats and got them working in the non instance view
	- Made tunnel stats inclusive rather than average and updated the code a little ( set average as 0.0, so the samples work but with no discrete timings )
	- Fixed complex tunnel widget generation ( it was missig the input pin links )
	- Fixed simple tunnel widget generation
	- Moved pure chain exec nodes over to being FScriptExecutionPureNode's
	- Fixed blueprint pure timings
	- Renamed the BaseSample in the scriptperfdata to SampleFrequency ( I couldn't think of a good name the other day )
	- Fixed a couple of issues I introduced into the FScriptPerfData updates that blocked blueprint stat updates working.
	- Made average stats always show when samples are present.

Change 3056373 on 2016/07/19 by [REDACTED_PERSON]

Change 3057320 on 2016/07/20 by [REDACTED_PERSON]

	FArchiveReplaceObjectRef works without recursion (it caused Stack overflow for very big BP).

Change 3057341 on 2016/07/20 by [REDACTED_PERSON]

Change 3057448 on 2016/07/20 by [REDACTED_PERSON]

	Any literal, passed as parameter to any Custom Thunk function, will have a proper type.

Change 3057764 on 2016/07/20 by [REDACTED_PERSON]

	This fixes a problem with the blueprint profiler mapping actor instances created at runtime to fail mapping and assert.
	#Proj BlueprintProfiler

diff --git a/Engine/Source/Developer/BlueprintCompilerCppBackend/Private/BlueprintCompilerCppBackend.cpp b/Engine/Source/Developer/BlueprintCompilerCppBackend/Private/BlueprintCompilerCppBackend.cpp
--- a/Engine/Source/Developer/BlueprintCompilerCppBackend/Private/BlueprintCompilerCppBackend.cpp
+++ b/Engine/Source/Developer/BlueprintCompilerCppBackend/Private/BlueprintCompilerCppBackend.cpp
@@ -924,13 +924,16 @@ FString FBlueprintCompilerCppBackend::TermToText(FEmitterLocalContext& EmitterCo
 
 		if (!Conditions.IsEmpty())
 		{
-			const FString DefaultValueConstructor = FEmitHelper::DefaultValue(EmitterContext, Term->Type);
 			const FString DefaultValueVariable = EmitterContext.GenerateUniqueLocalName();
 			const uint32 PropertyExportFlags = EPropertyExportCPPFlags::CPPF_CustomTypeName | EPropertyExportCPPFlags::CPPF_BlueprintCppBackend | EPropertyExportCPPFlags::CPPF_NoConst;
 			const FString CppType = Term->AssociatedVarProperty 
 				? EmitterContext.ExportCppDeclaration(Term->AssociatedVarProperty, EExportedDeclaration::Local, PropertyExportFlags, true)
 				: FEmitHelper::PinTypeToNativeType(Term->Type);
 
+			const FString DefaultValueConstructor = (!Term->Type.bIsArray)
+				? FEmitHelper::LiteralTerm(EmitterContext, Term->Type, FString(), nullptr, &FText::GetEmpty())
+				: FString::Printf(TEXT("%s{}"), *CppType);
+
 			EmitterContext.AddLine(*FString::Printf(TEXT("%s %s = %s;"), *CppType, *DefaultValueVariable, *DefaultValueConstructor));
 			return FString::Printf(TEXT("((%s) ? (%s) : (%s))"), *Conditions, *ResultPath, *DefaultValueVariable);
 		}
diff --git a/Engine/Source/Developer/BlueprintCompilerCppBackend/Private/BlueprintCompilerCppBackendBase.cpp b/Engine/Source/Developer/BlueprintCompilerCppBackend/Private/BlueprintCompilerCppBackendBase.cpp
--- a/Engine/Source/Developer/BlueprintCompilerCppBackend/Private/BlueprintCompilerCppBackendBase.cpp
+++ b/Engine/Source/Developer/BlueprintCompilerCppBackend/Private/BlueprintCompilerCppBackendBase.cpp
@@ -281,14 +281,23 @@ FString FBlueprintCompilerCppBackendBase::GenerateCodeFromClass(UClass* SourceCl
 
 	EmitStructProperties(EmitterContext, SourceClass);
 
+	TSharedPtr<FGatherConvertedClassDependencies> ParentDependencies;
 	// Emit function declarations and definitions (writes to header and body simultaneously)
 	if (!bIsInterface)
 	{
+		UBlueprintGeneratedClass* BPGC = CastChecked<UBlueprintGeneratedClass>(EmitterContext.GetCurrentlyGeneratedClass());
+		UBlueprintGeneratedClass* ParentBPGC = Cast<UBlueprintGeneratedClass>(BPGC->GetSuperClass());
+		ParentDependencies = TSharedPtr<FGatherConvertedClassDependencies>(ParentBPGC ? new FGatherConvertedClassDependencies(ParentBPGC) : nullptr);
+
+		FEmitDefaultValueHelper::FillCommonUsedAssets(EmitterContext, ParentDependencies);
+
 		EmitterContext.Header.AddLine(FString::Printf(TEXT("%s(const FObjectInitializer& ObjectInitializer = FObjectInitializer::Get());"), *CppClassName));
 		EmitterContext.Header.AddLine(TEXT("virtual void PostLoadSubobjects(FObjectInstancingGraph* OuterInstanceGraph) override;"));
-		EmitterContext.Header.AddLine(TEXT("static void __StaticDependenciesAssets(TArray<FBlueprintDependencyData>& AssetsToLoad);"));
 		EmitterContext.Header.AddLine(TEXT("static void __CustomDynamicClassInitialization(UDynamicClass* InDynamicClass);"));
-		EmitterContext.Header.AddLine(TEXT("static void __CustomDynamicClassInitializationUsedAssets(UDynamicClass* InDynamicClass);"));
+
+		EmitterContext.Header.AddLine(TEXT("static void __StaticDependenciesAssets(TArray<FBlueprintDependencyData>& AssetsToLoad);"));
+		EmitterContext.Header.AddLine(TEXT("static void __StaticDependencies_CommonAssets(TArray<FBlueprintDependencyData>& AssetsToLoad);"));
+		EmitterContext.Header.AddLine(TEXT("static void __StaticDependencies_DirectlyUsedAssets(TArray<FBlueprintDependencyData>& AssetsToLoad);"));
 		if (bHasStaticSearchableValues)
 		{
 			FBackendHelperStaticSearchableValues::EmitFunctionDeclaration(EmitterContext);
@@ -325,7 +334,9 @@ FString FBlueprintCompilerCppBackendBase::GenerateCodeFromClass(UClass* SourceCl
 	{
 		// must be called after GenerateConstructor
 		// now we knows which assets are directly used in source code
-		FEmitDefaultValueHelper::GenerateCustomDynamicClassInitializationUsedAssets(EmitterContext);
+		FEmitDefaultValueHelper::GenerateCustomDynamicClassInitialization(EmitterContext, ParentDependencies);
+		FEmitDefaultValueHelper::AddStaticFunctionsForDependencies(EmitterContext, ParentDependencies);
+		FEmitDefaultValueHelper::AddRegisterHelper(EmitterContext);
 	}
 
 	FEmitHelper::EmitLifetimeReplicatedPropsImpl(EmitterContext);
diff --git a/Engine/Source/Developer/BlueprintCompilerCppBackend/Private/BlueprintCompilerCppBackendGatherDependencies.cpp b/Engine/Source/Developer/BlueprintCompilerCppBackend/Private/BlueprintCompilerCppBackendGatherDependencies.cpp
--- a/Engine/Source/Developer/BlueprintCompilerCppBackend/Private/BlueprintCompilerCppBackendGatherDependencies.cpp
+++ b/Engine/Source/Developer/BlueprintCompilerCppBackend/Private/BlueprintCompilerCppBackendGatherDependencies.cpp
@@ -280,8 +280,9 @@ FGatherConvertedClassDependencies::FGatherConvertedClassDependencies(UStruct* In
 	static const FBoolConfigValueHelper DontNativizeDataOnlyBP(TEXT("BlueprintNativizationSettings"), TEXT("bDontNativizeDataOnlyBP"));
 	if (DontNativizeDataOnlyBP)
 	{
-		auto RemoveFieldsFromDataOnlyBP = [](TSet<UField*>& FieldSet)
+		auto RemoveFieldsFromDataOnlyBP = [&](TSet<UField*>& FieldSet)
 		{
+			TSet<UField*> FieldsToAdd;
 			for (auto Iter = FieldSet.CreateIterator(); Iter; ++Iter)
 			{
 				UClass* CurrentClass = (*Iter) ? (*Iter)->GetOwnerClass() : nullptr;
@@ -289,8 +290,11 @@ FGatherConvertedClassDependencies::FGatherConvertedClassDependencies(UStruct* In
 				if (CurrentBP && FBlueprintEditorUtils::IsDataOnlyBlueprint(CurrentBP))
 				{
 					Iter.RemoveCurrent();
+					FieldsToAdd.Add(GetFirstNativeOrConvertedClass(CurrentClass, true));
 				}
 			}
+
+			FieldSet.Append(FieldsToAdd);
 		};
 		RemoveFieldsFromDataOnlyBP(IncludeInHeader);
 		RemoveFieldsFromDataOnlyBP(DeclareInHeader);
@@ -298,6 +302,29 @@ FGatherConvertedClassDependencies::FGatherConvertedClassDependencies(UStruct* In
 	}
 }
 
+UClass* FGatherConvertedClassDependencies::GetFirstNativeOrConvertedClass(UClass* InClass, bool bExcludeBPDataOnly) const
+{
+	check(InClass);
+	for (UClass* ItClass = InClass; ItClass; ItClass = ItClass->GetSuperClass())
+	{
+		UBlueprintGeneratedClass* BPGC = Cast<UBlueprintGeneratedClass>(ItClass);
+		if (ItClass->HasAnyClassFlags(CLASS_Native) || WillClassBeConverted(BPGC))
+		{
+			if (bExcludeBPDataOnly)
+			{
+				UBlueprint* BP = BPGC ? Cast<UBlueprint>(BPGC->ClassGeneratedBy) : nullptr;
+				if (BP && FBlueprintEditorUtils::IsDataOnlyBlueprint(BP))
+				{
+					continue;
+				}
+			}
+			return ItClass;
+		}
+	}
+	check(false);
+	return nullptr;
+};
+
 UClass* FGatherConvertedClassDependencies::FindOriginalClass(const UClass* InClass) const
 {
 	if (InClass)
@@ -367,22 +394,7 @@ void FGatherConvertedClassDependencies::DependenciesForHeader()
 		{
 			if (auto ObjectProperty = Cast<const UObjectPropertyBase>(Property))
 			{
-				auto GetFirstNativeOrConvertedClass = [](FGatherConvertedClassDependencies& Dependencies, UClass* InClass) -> UClass*
-				{
-					check(InClass);
-					for (UClass* ItClass = InClass; ItClass; ItClass = ItClass->GetSuperClass())
-					{
-						auto BPGC = Cast<UBlueprintGeneratedClass>(ItClass);
-						if (ItClass->HasAnyClassFlags(CLASS_Native) || Dependencies.WillClassBeConverted(BPGC))
-						{
-							return ItClass;
-						}
-					}
-					check(false);
-					return nullptr;
-				};
-
-				DeclareInHeader.Add(GetFirstNativeOrConvertedClass(*this, ObjectProperty->PropertyClass));
+				DeclareInHeader.Add(GetFirstNativeOrConvertedClass(ObjectProperty->PropertyClass));
 			}
 			else if (auto InterfaceProperty = Cast<const UInterfaceProperty>(Property))
 			{
@@ -431,10 +443,24 @@ void FGatherConvertedClassDependencies::DependenciesForHeader()
 		}
 	}
 
+	// DEFAULT VALUES FROM UDS:
+	UUserDefinedStruct* UDS = Cast<UUserDefinedStruct>(OriginalStruct);
+	if (UDS)
+	{
+		FStructOnScope StructOnScope(UDS);
+		UDS->InitializeDefaultValue(StructOnScope.GetStructMemory());
+		for (TFieldIterator<UObjectPropertyBase> PropertyIt(UDS); PropertyIt; ++PropertyIt)
+		{
+			UObject* DefaultValueObject = ((UObjectPropertyBase*)*PropertyIt)->GetObjectPropertyValue_InContainer(StructOnScope.GetStructMemory());
+			UField* ObjAsField = Cast<UField>(DefaultValueObject);
+			UField* FieldForHeader = ObjAsField ? ObjAsField : (DefaultValueObject ? DefaultValueObject->GetClass() : nullptr);
+			IncludeInHeader.Add(FieldForHeader);
+		}
+	}
 
+	// REMOVE UNNECESSARY HEADERS
 	UClass* AsBPGC = Cast<UBlueprintGeneratedClass>(OriginalStruct);
 	UClass* OriginalClassFromOriginalPackage = AsBPGC ? FindOriginalClass(AsBPGC) : nullptr;
-
 	const UPackage* OriginalStructPackage = OriginalStruct ? OriginalStruct->GetOutermost() : nullptr;
 	for (auto Iter = IncludeInHeader.CreateIterator(); Iter; ++Iter)
 	{
diff --git a/Engine/Source/Developer/BlueprintCompilerCppBackend/Private/BlueprintCompilerCppBackendUMG.cpp b/Engine/Source/Developer/BlueprintCompilerCppBackend/Private/BlueprintCompilerCppBackendUMG.cpp
--- a/Engine/Source/Developer/BlueprintCompilerCppBackend/Private/BlueprintCompilerCppBackendUMG.cpp
+++ b/Engine/Source/Developer/BlueprintCompilerCppBackend/Private/BlueprintCompilerCppBackendUMG.cpp
@@ -84,7 +84,7 @@ void FBackendHelperUMG::EmitWidgetInitializationFunctions(FEmitterLocalContext&
 			const FString AnimationsArrayNativeName = GenerateLocalProperty(Context, FindFieldChecked<UArrayProperty>(UWidgetBlueprintGeneratedClass::StaticClass(), TEXT("Animations")), reinterpret_cast<const uint8*>(&WidgetClass->Animations));
 			const FString BindingsArrayNativeName = GenerateLocalProperty(Context, FindFieldChecked<UArrayProperty>(UWidgetBlueprintGeneratedClass::StaticClass(), TEXT("Bindings")), reinterpret_cast<const uint8*>(&WidgetClass->Bindings));
 					
-			Context.AddLine(FString::Printf(TEXT("UWidgetBlueprintGeneratedClass::%s(this, GetClass(), %s, %s, %s, %s, %s, nullptr);")
+			Context.AddLine(FString::Printf(TEXT("UWidgetBlueprintGeneratedClass::%s(this, GetClass(), %s, %s, %s, %s, %s);")
 				, GET_FUNCTION_NAME_STRING_CHECKED(UWidgetBlueprintGeneratedClass, InitializeWidgetStatic)
 				, WidgetClass->bCanEverTick ? TEXT("true") : TEXT("false")
 				, WidgetClass->bCanEverPaint ? TEXT("true") : TEXT("false")
diff --git a/Engine/Source/Developer/BlueprintCompilerCppBackend/Private/BlueprintCompilerCppBackendUtils.cpp b/Engine/Source/Developer/BlueprintCompilerCppBackend/Private/BlueprintCompilerCppBackendUtils.cpp
--- a/Engine/Source/Developer/BlueprintCompilerCppBackend/Private/BlueprintCompilerCppBackendUtils.cpp
+++ b/Engine/Source/Developer/BlueprintCompilerCppBackend/Private/BlueprintCompilerCppBackendUtils.cpp
@@ -1009,10 +1009,18 @@ FString FEmitHelper::LiteralTerm(FEmitterLocalContext& EmitterContext, const FEd
 					return CustomConstructor;
 				}
 
-				EmitterContext.AddLine(FString::Printf(TEXT("auto %s = %s{};"), *LocalStructNativeName, *StructName)); // TODO: ?? should "::GetDefaultValue()" be called here?
-				for (auto LocalProperty : TFieldRange<const UProperty>(StructType))
 				{
-					FEmitDefaultValueHelper::OuterGenerate(EmitterContext, LocalProperty, LocalStructNativeName, StructOnScope.GetStructMemory(), nullptr, FEmitDefaultValueHelper::EPropertyAccessOperator::Dot);
+					const FString StructOnScopeName = EmitterContext.GenerateUniqueLocalName();
+					EmitterContext.AddLine(FString::Printf(TEXT("FStructOnScope %s(%s::StaticStruct());"), *StructOnScopeName, *StructName));
+					EmitterContext.AddLine(FString::Printf(TEXT("%s& %s = *((%s*)%s.GetStructMemory());"), *StructName, *LocalStructNativeName, *StructName, *StructOnScopeName));  // TODO: ?? should "::GetDefaultValue()" be called here?
+				}
+
+				{
+					FStructOnScope DefaultStructOnScope(StructType);
+					for (auto LocalProperty : TFieldRange<const UProperty>(StructType))
+					{
+						FEmitDefaultValueHelper::OuterGenerate(EmitterContext, LocalProperty, LocalStructNativeName, StructOnScope.GetStructMemory(), DefaultStructOnScope.GetStructMemory(), FEmitDefaultValueHelper::EPropertyAccessOperator::Dot);
+					}
 				}
 			}
 			return LocalStructNativeName;
@@ -1185,17 +1193,6 @@ FString FEmitHelper::PinTypeToNativeType(const FEdGraphPinType& Type)
 	return Type.bIsArray ? FString::Printf(TEXT("TArray<%s>"), *InnerTypeName) : InnerTypeName;
 }
 
-FString FEmitHelper::DefaultValue(FEmitterLocalContext& EmitterContext, const FEdGraphPinType& Type)
-{
-	if (Type.bIsArray)
-	{
-		const FString InnerTypeName = PinTypeToNativeType(Type);
-		return FString::Printf(TEXT("%s{}"), *InnerTypeName);
-	}
-
-	return LiteralTerm(EmitterContext, Type, FString(), nullptr, &FText::GetEmpty());
-}
-
 UFunction* FEmitHelper::GetOriginalFunction(UFunction* Function)
 {
 	check(Function);
diff --git a/Engine/Source/Developer/BlueprintCompilerCppBackend/Private/BlueprintCompilerCppBackendUtils.h b/Engine/Source/Developer/BlueprintCompilerCppBackend/Private/BlueprintCompilerCppBackendUtils.h
--- a/Engine/Source/Developer/BlueprintCompilerCppBackend/Private/BlueprintCompilerCppBackendUtils.h
+++ b/Engine/Source/Developer/BlueprintCompilerCppBackend/Private/BlueprintCompilerCppBackendUtils.h
@@ -142,17 +142,7 @@ public:
 
 	UClass* GetFirstNativeOrConvertedClass(UClass* InClass) const
 	{
-		check(InClass);
-		for (UClass* ItClass = InClass; ItClass; ItClass = ItClass->GetSuperClass())
-		{
-			auto BPGC = Cast<UBlueprintGeneratedClass>(ItClass);
-			if (ItClass->HasAnyClassFlags(CLASS_Native) || (ensure(BPGC) && Dependencies.WillClassBeConverted(BPGC)))
-			{
-				return ItClass;
-			}
-		}
-		check(false);
-		return nullptr;
+		return Dependencies.GetFirstNativeOrConvertedClass(InClass);
 	}
 
 	FString GenerateUniqueLocalName();
@@ -236,8 +226,6 @@ struct FEmitHelper
 
 	static FString LiteralTerm(FEmitterLocalContext& EmitterContext, const FEdGraphPinType& Type, const FString& CustomValue, UObject* LiteralObject, const FText* OptionalTextLiteral = nullptr);
 
-	static FString DefaultValue(FEmitterLocalContext& EmitterContext, const FEdGraphPinType& Type);
-
 	static FString PinTypeToNativeType(const FEdGraphPinType& InType);
 
 	static UFunction* GetOriginalFunction(UFunction* Function);
@@ -278,9 +266,9 @@ struct FEmitDefaultValueHelper
 
 	static void GenerateConstructor(FEmitterLocalContext& Context);
 
-	static void GenerateCustomDynamicClassInitialization(FEmitterLocalContext& Context, TSharedPtr<FGatherConvertedClassDependencies> ParentDependencies);
+	static void FillCommonUsedAssets(FEmitterLocalContext& Context, TSharedPtr<FGatherConvertedClassDependencies> ParentDependencies);
 
-	static void GenerateCustomDynamicClassInitializationUsedAssets(FEmitterLocalContext& Context);
+	static void GenerateCustomDynamicClassInitialization(FEmitterLocalContext& Context, TSharedPtr<FGatherConvertedClassDependencies> ParentDependencies);
 
 	enum class EPropertyAccessOperator
 	{
@@ -303,6 +291,10 @@ struct FEmitDefaultValueHelper
 	// returns true, and fill OutResult, when the structure is handled in a custom way.
 	static bool SpecialStructureConstructor(const UStruct* Struct, const uint8* ValuePtr, FString* OutResult);
 
+	// Add static initialization functions. Must be call after Context.UsedObjectInCurrentClass is fully filled
+	static void AddStaticFunctionsForDependencies(FEmitterLocalContext& Context, TSharedPtr<FGatherConvertedClassDependencies> ParentDependencies);
+
+	static void AddRegisterHelper(FEmitterLocalContext& Context);
 private:
 	// Returns native term, 
 	// returns empty string if cannot handle
diff --git a/Engine/Source/Developer/BlueprintCompilerCppBackend/Private/BlueprintCompilerCppBackendValueHelper.cpp b/Engine/Source/Developer/BlueprintCompilerCppBackend/Private/BlueprintCompilerCppBackendValueHelper.cpp
--- a/Engine/Source/Developer/BlueprintCompilerCppBackend/Private/BlueprintCompilerCppBackendValueHelper.cpp
+++ b/Engine/Source/Developer/BlueprintCompilerCppBackend/Private/BlueprintCompilerCppBackendValueHelper.cpp
@@ -120,17 +120,17 @@ void FEmitDefaultValueHelper::GenerateGetDefaultValue(const UUserDefinedStruct*
 	Context.Header.AddLine(TEXT("{"));
 
 	Context.Header.IncreaseIndent();
-	Context.Header.AddLine(FString::Printf(TEXT("%s DefaultData__;"), *StructName));
+	Context.Header.AddLine(FString::Printf(TEXT("FStructOnScope StructOnScope(%s::StaticStruct());"), *StructName));
+	Context.Header.AddLine(FString::Printf(TEXT("%s& DefaultData__ = *((%s*)StructOnScope.GetStructMemory());"), *StructName, *StructName));
 	{
 		TGuardValue<FCodeText*> OriginalDefaultTarget(Context.DefaultTarget, &Context.Header);
 		FStructOnScope StructData(Struct);
 		FStructureEditorUtils::Fill_MakeStructureDefaultValue(Struct, StructData.GetStructMemory());
+		FStructOnScope RawDefaultStructOnScope(Struct);
 		for (auto Property : TFieldRange<const UProperty>(Struct))
 		{
-			OuterGenerate(Context, Property, TEXT("DefaultData__"), StructData.GetStructMemory(), nullptr, EPropertyAccessOperator::Dot);
+			OuterGenerate(Context, Property, TEXT("DefaultData__"), StructData.GetStructMemory(), RawDefaultStructOnScope.GetStructMemory(), EPropertyAccessOperator::Dot);
 		}
-
-
 	}
 	Context.Header.AddLine(TEXT("return DefaultData__;"));
 	Context.Header.DecreaseIndent();
@@ -188,9 +188,14 @@ void FEmitDefaultValueHelper::InnerGenerate(FEmitterLocalContext& Context, const
 
 	if (StructProperty)
 	{
+		// Create default struct instance, only when DefaultValuePtr is null.
+		FStructOnScope DefaultStructOnScope(DefaultValuePtr ? nullptr : StructProperty->Struct);
+
 		for (auto LocalProperty : TFieldRange<const UProperty>(StructProperty->Struct))
 		{
-			OuterGenerate(Context, LocalProperty, PathToMember, ValuePtr, DefaultValuePtr, EPropertyAccessOperator::Dot);
+			OuterGenerate(Context, LocalProperty, PathToMember, ValuePtr
+				, (DefaultValuePtr ? DefaultValuePtr : DefaultStructOnScope.GetStructMemory())
+				, EPropertyAccessOperator::Dot);
 		}
 	}
 	
@@ -246,7 +251,13 @@ void FEmitDefaultValueHelper::InnerGenerate(FEmitterLocalContext& Context, const
 			
 			if (RegularInnerStruct || (bInitializeWithoutScriptStruct && !bComplete))
 			{
-				const FString LocalPathToMember = FString::Printf(TEXT("%s[%d]"), *PathToMember, Index);
+				FString LocalPathToMember = FString::Printf(TEXT("%s[%d]"), *PathToMember, Index);
+				if (RegularInnerStruct)
+				{
+					const FString ArrayElementRefName = Context.GenerateUniqueLocalName();
+					Context.AddLine(FString::Printf(TEXT("auto& %s = %s;"), *ArrayElementRefName, *LocalPathToMember));
+					LocalPathToMember = ArrayElementRefName;
+				}
 				InnerGenerate(Context, ArrayProperty->Inner, LocalPathToMember, LocalValuePtr, DefaultStruct.GetStructMemory(), true);
 			}
 		}
@@ -257,6 +268,19 @@ bool FEmitDefaultValueHelper::SpecialStructureConstructor(const UStruct* Struct,
 {
 	check(ValuePtr || !OutResult);
 
+	if (FLatentActionInfo::StaticStruct() == Struct)
+	{
+		if (OutResult)
+		{
+			const FLatentActionInfo* LatentActionInfo = reinterpret_cast<const FLatentActionInfo*>(ValuePtr);
+			*OutResult = FString::Printf(TEXT("FLatentActionInfo(%d, %d, TEXT(\"%s\"), this)")
+				, LatentActionInfo->Linkage
+				, LatentActionInfo->UUID
+				, *LatentActionInfo->ExecutionFunction.ToString().ReplaceCharWithEscapedChar());
+		}
+		return true;
+	}
+
 	if (TBaseStructure<FTransform>::Get() == Struct)
 	{
 		if (OutResult)
@@ -560,16 +584,45 @@ struct FNonativeComponentData
 			// AttachTo is called first in case some properties will be overridden.
 		}
 
+		bool bBodyInstanceIsAlreadyHandled = false;
+		UProperty* BodyInstanceProperty = UPrimitiveComponent::StaticClass()->FindPropertyByName(GET_MEMBER_NAME_CHECKED(UPrimitiveComponent, BodyInstance));
+		UPrimitiveComponent* PrimitiveComponent = Cast<UPrimitiveComponent>(ComponentTemplate);
+		if(PrimitiveComponent)
+		{
+			const FName CollisionProfileName = PrimitiveComponent->BodyInstance.GetCollisionProfileName();
+			UPrimitiveComponent* ComponentArchetype = Cast<UPrimitiveComponent>(ObjectToCompare);
+			const FName ComponentArchetypeCollisionProfileName = ComponentArchetype ? ComponentArchetype->BodyInstance.GetCollisionProfileName() : NAME_None;
+			if (CollisionProfileName != ComponentArchetypeCollisionProfileName)
+			{
+				FStructOnScope BodyInstanceToCompare(FBodyInstance::StaticStruct());
+				if (ComponentArchetype)
+				{
+					FBodyInstance::StaticStruct()->CopyScriptStruct(BodyInstanceToCompare.GetStructMemory(), &ComponentArchetype->BodyInstance);
+				}
+				((FBodyInstance*)BodyInstanceToCompare.GetStructMemory())->SetCollisionProfileName(CollisionProfileName);
+
+				const FString PathToMember = FString::Printf(TEXT("%s->BodyInstance"), *NativeVariablePropertyName);
+				Context.AddLine(FString::Printf(TEXT("%s.SetCollisionProfileName(FName(TEXT(\"%s\")));"), *PathToMember, *CollisionProfileName.ToString().ReplaceCharWithEscapedChar()));
+				FEmitDefaultValueHelper::InnerGenerate(Context, BodyInstanceProperty, PathToMember, (const uint8*)&PrimitiveComponent->BodyInstance, BodyInstanceToCompare.GetStructMemory());
+				bBodyInstanceIsAlreadyHandled = true;
+			}
+		}
+
 		UClass* ComponentClass = ComponentTemplate->GetClass();
 		for (auto Property : TFieldRange<const UProperty>(ComponentClass))
 		{
-			if (!HandledAsSpecialProperty(Context, Property))
+			if (bBodyInstanceIsAlreadyHandled && (Property == BodyInstanceProperty))
 			{
-				FEmitDefaultValueHelper::OuterGenerate(Context, Property, NativeVariablePropertyName
-					, reinterpret_cast<const uint8*>(ComponentTemplate)
-					, reinterpret_cast<const uint8*>(ObjectToCompare)
-					, FEmitDefaultValueHelper::EPropertyAccessOperator::Pointer);
+				continue;
 			}
+			if (HandledAsSpecialProperty(Context, Property))
+			{
+				continue;
+			}
+			FEmitDefaultValueHelper::OuterGenerate(Context, Property, NativeVariablePropertyName
+				, reinterpret_cast<const uint8*>(ComponentTemplate)
+				, reinterpret_cast<const uint8*>(ObjectToCompare)
+				, FEmitDefaultValueHelper::EPropertyAccessOperator::Pointer);
 		}
 	}
 
@@ -710,54 +763,22 @@ public:
 		const FString ClassString = Item->IsA<UClass>() ? TEXT("UClass") : TEXT("UScriptStruct");
 		return FString(TEXT("Z_Construct_")) + ClassString + TEXT("_") + Result + TEXT("()");
 	}
+};
 
-	static void AddStaticFunctionsForDependencies(FEmitterLocalContext& Context, TSharedPtr<FGatherConvertedClassDependencies> ParentDependencies)
-	{
-		auto SourceClass = Context.GetCurrentlyGeneratedClass();
-		auto OriginalClass = Context.Dependencies.FindOriginalClass(SourceClass);
-		const FString CppClassName = FEmitHelper::GetCppName(OriginalClass);
-		
-		// __StaticDependenciesAssets
-		Context.AddLine(FString::Printf(TEXT("void %s::__StaticDependenciesAssets(TArray<FBlueprintDependencyData>& AssetsToLoad)"), *CppClassName));
-		Context.AddLine(TEXT("{"));
-		Context.IncreaseIndent();
-
-		UBlueprintGeneratedClass* ParentBPGC = Cast<UBlueprintGeneratedClass>(SourceClass->GetSuperClass());
-		ensure(ParentDependencies.IsValid() == !!ParentBPGC);
-		if (ParentBPGC)
-		{
-			Context.AddLine(FString::Printf(TEXT("%s::__StaticDependenciesAssets(AssetsToLoad);"), *FEmitHelper::GetCppName(ParentBPGC)));
-		}
+void FEmitDefaultValueHelper::AddStaticFunctionsForDependencies(FEmitterLocalContext& Context, TSharedPtr<FGatherConvertedClassDependencies> ParentDependencies)
+{
 
-		auto CreateAssetToLoadString = [&](const UObject* AssetObj) -> FString
-		{
-			UClass* AssetType = AssetObj->GetClass();
-			if (AssetType->IsChildOf<UUserDefinedEnum>())
-			{
-				AssetType = UEnum::StaticClass();
-			}
-			else if (AssetType->IsChildOf<UUserDefinedStruct>())
-			{
-				AssetType = UScriptStruct::StaticClass();
-			}
-			else if (AssetType->IsChildOf<UBlueprintGeneratedClass>() && Context.Dependencies.WillClassBeConverted(CastChecked<UBlueprintGeneratedClass>(AssetObj)))
-			{
-				AssetType = UDynamicClass::StaticClass();
-			}
 
-			return FString::Printf(TEXT("AssetsToLoad.Add({FName(TEXT(\"%s\")), FName(TEXT(\"%s\")), FName(TEXT(\"%s\")), FName(TEXT(\"%s\"))});")
-				, *AssetObj->GetOutermost()->GetPathName()
-				, *AssetObj->GetName()
-				, *AssetType->GetOutermost()->GetPathName()
-				, *AssetType->GetName());
-		};
+	// 0.ALL ASSETS TO LIST
+	TSet<const UObject*> StaticDependenciesAssets;
+	{
 		for (UObject* LocAsset : Context.Dependencies.Assets)
 		{
 			if (ParentDependencies.IsValid() && ParentDependencies->Assets.Contains(LocAsset))
 			{
 				continue;
 			}
-			Context.AddLine(CreateAssetToLoadString(LocAsset));
+			StaticDependenciesAssets.Add(LocAsset);
 		}
 		for (UBlueprintGeneratedClass* LocAsset : Context.Dependencies.ConvertedClasses)
 		{
@@ -765,60 +786,175 @@ public:
 			{
 				continue;
 			}
-			if (!Context.Dependencies.Assets.Contains(LocAsset))
+			StaticDependenciesAssets.Add(LocAsset);
+		}
+	}
+
+	// 1. USED ASSETS LIST
+	auto SourceClass = Context.GetCurrentlyGeneratedClass();
+	auto OriginalClass = Context.Dependencies.FindOriginalClass(SourceClass);
+	const FString CppClassName = FEmitHelper::GetCppName(OriginalClass);
+
+	auto CreateAssetToLoadString = [&](const UObject* AssetObj) -> FString
+	{
+		UClass* AssetType = AssetObj->GetClass();
+		if (AssetType->IsChildOf<UUserDefinedEnum>())
+		{
+			AssetType = UEnum::StaticClass();
+		}
+		else if (AssetType->IsChildOf<UUserDefinedStruct>())
+		{
+			AssetType = UScriptStruct::StaticClass();
+		}
+		else if (AssetType->IsChildOf<UBlueprintGeneratedClass>() && Context.Dependencies.WillClassBeConverted(CastChecked<UBlueprintGeneratedClass>(AssetObj)))
+		{
+			AssetType = UDynamicClass::StaticClass();
+		}
+
+		return FString::Printf(TEXT("FBlueprintDependencyData(TEXT(\"%s\"), TEXT(\"%s\"), TEXT(\"%s\"), TEXT(\"%s\")),")
+			, *AssetObj->GetOutermost()->GetPathName()
+			, *AssetObj->GetName()
+			, *AssetType->GetOutermost()->GetPathName()
+			, *AssetType->GetName());
+	};
+
+	auto AddAssetArray = [&](TArray<const UObject*>& Assets)
+	{
+		if (Assets.Num())
+		{
+			Context.AddLine(TEXT("FBlueprintDependencyData LocAssets[] ="));
+			Context.AddLine(TEXT("{"));
+			Context.IncreaseIndent();
+		}
+
+		for (const UObject* LocAsset : Assets)
+		{
+			Context.AddLine(CreateAssetToLoadString(LocAsset));
+		}
+
+		if (Assets.Num())
+		{
+			Context.DecreaseIndent();
+			Context.AddLine(TEXT("};"));
+			Context.AddLine(TEXT("for(auto& LocAsset : LocAssets) { AssetsToLoad.Add(LocAsset); } "));
+		}
+	};
+
+	int32 UsedAssetIndex = 0;
+	// 1. COMMON DEPENDENCIES ASSETS - USED BOTH BY __StaticDependenciesAssets and __StaticDependencies_DirectlyUsedAssets
+	{
+		Context.AddLine(FString::Printf(TEXT("void %s::__StaticDependencies_CommonAssets(TArray<FBlueprintDependencyData>& AssetsToLoad)"), *CppClassName));
+		Context.AddLine(TEXT("{"));
+		Context.IncreaseIndent();
+
+		TArray<const UObject*> AssetsToAdd;
+		for (;UsedAssetIndex < Context.UsedObjectInCurrentClass.Num(); ++UsedAssetIndex)
+		{
+			const UObject* LocAsset = Context.UsedObjectInCurrentClass[UsedAssetIndex];
+			if (!StaticDependenciesAssets.Contains(LocAsset))
 			{
-				Context.AddLine(CreateAssetToLoadString(LocAsset));
+				break;
 			}
+			StaticDependenciesAssets.Remove(LocAsset);
+			AssetsToAdd.Add(LocAsset);
 		}
+		AddAssetArray(AssetsToAdd);
 
 		Context.DecreaseIndent();
 		Context.AddLine(TEXT("}"));
+	}
 
-		// Register Helper Struct
-		const FString RegisterHelperName = FString::Printf(TEXT("FRegisterHelper__%s"), *CppClassName);
-		Context.AddLine(FString::Printf(TEXT("struct %s"), *RegisterHelperName));
+	//
+	{
+		// __StaticDependencies_DirectlyUsedAssets
+		Context.AddLine(FString::Printf(TEXT("void %s::__StaticDependencies_DirectlyUsedAssets(TArray<FBlueprintDependencyData>& AssetsToLoad)"), *CppClassName));
 		Context.AddLine(TEXT("{"));
 		Context.IncreaseIndent();
 
-		Context.AddLine(FString::Printf(TEXT("%s()"), *RegisterHelperName));
-		Context.AddLine(TEXT("{"));
-		Context.IncreaseIndent();
+		Context.AddLine(FString(TEXT("__StaticDependencies_CommonAssets(AssetsToLoad);")));
 
-		Context.AddLine(FString::Printf(
-			TEXT("FConvertedBlueprintsDependencies::Get().RegisterClass(TEXT(\"%s\"), &%s::__StaticDependenciesAssets);")
-			, *OriginalClass->GetOutermost()->GetPathName()
-			, *CppClassName));
+		TArray<const UObject*> AssetsToAdd;
+		for (; UsedAssetIndex < Context.UsedObjectInCurrentClass.Num(); ++UsedAssetIndex)
+		{
+			const UObject* LocAsset = Context.UsedObjectInCurrentClass[UsedAssetIndex];
+			ensure(!StaticDependenciesAssets.Contains(LocAsset));
+			AssetsToAdd.Add(LocAsset);
+		}
+		AddAssetArray(AssetsToAdd);
 
 		Context.DecreaseIndent();
 		Context.AddLine(TEXT("}"));
+	}
 
-		Context.AddLine(FString::Printf(TEXT("static %s Instance;"), *RegisterHelperName));
+	// 2. REMAINING DEPENDENCIES
+	// __StaticDependenciesAssets
+	{
+		Context.AddLine(FString::Printf(TEXT("void %s::__StaticDependenciesAssets(TArray<FBlueprintDependencyData>& AssetsToLoad)"), *CppClassName));
+		Context.AddLine(TEXT("{"));
+		Context.IncreaseIndent();
 
-		Context.DecreaseIndent();
-		Context.AddLine(TEXT("};"));
+		UBlueprintGeneratedClass* ParentBPGC = Cast<UBlueprintGeneratedClass>(SourceClass->GetSuperClass());
+		ensure(ParentDependencies.IsValid() == !!ParentBPGC);
+		if (ParentBPGC)
+		{
+			Context.AddLine(FString::Printf(TEXT("%s::__StaticDependenciesAssets(AssetsToLoad);"), *FEmitHelper::GetCppName(ParentBPGC)));
+		}
+		Context.AddLine(FString(TEXT("__StaticDependencies_CommonAssets(AssetsToLoad);")));
+
+		TArray<const UObject*> AssetsToAdd;
+		for (const UObject* LocAsset : StaticDependenciesAssets)
+		{
+			AssetsToAdd.Add(LocAsset);
+		}
+		AddAssetArray(AssetsToAdd);
 
-		Context.AddLine(FString::Printf(TEXT("%s %s::Instance;"), *RegisterHelperName, *RegisterHelperName));
+		Context.DecreaseIndent();
+		Context.AddLine(TEXT("}"));
 	}
-};
+}
 
-void FEmitDefaultValueHelper::GenerateCustomDynamicClassInitializationUsedAssets(FEmitterLocalContext& Context)
+void FEmitDefaultValueHelper::AddRegisterHelper(FEmitterLocalContext& Context)
 {
-	auto BPGC = CastChecked<UBlueprintGeneratedClass>(Context.GetCurrentlyGeneratedClass());
-	const FString CppClassName = FEmitHelper::GetCppName(BPGC);
+	auto SourceClass = Context.GetCurrentlyGeneratedClass();
+	auto OriginalClass = Context.Dependencies.FindOriginalClass(SourceClass);
+	const FString CppClassName = FEmitHelper::GetCppName(OriginalClass);
 
-	Context.AddLine(FString::Printf(TEXT("void %s::__CustomDynamicClassInitializationUsedAssets(UDynamicClass* InDynamicClass)"), *CppClassName));
+	const FString RegisterHelperName = FString::Printf(TEXT("FRegisterHelper__%s"), *CppClassName);
+	Context.AddLine(FString::Printf(TEXT("struct %s"), *RegisterHelperName));
 	Context.AddLine(TEXT("{"));
 	Context.IncreaseIndent();
-	Context.AddLine(FString::Printf(TEXT("ensure(0 == InDynamicClass->%s.Num());"), GET_MEMBER_NAME_STRING_CHECKED(UDynamicClass, UsedAssets)));
 
-	for (auto LocAsset : Context.UsedObjectInCurrentClass)
-	{
-		const FString AssetStr = Context.FindGloballyMappedObject(LocAsset, UObject::StaticClass(), true, false);
-		Context.AddLine(FString::Printf(TEXT("InDynamicClass->%s.Add(%s);"), GET_MEMBER_NAME_STRING_CHECKED(UDynamicClass, UsedAssets), *AssetStr));
-	}
+	Context.AddLine(FString::Printf(TEXT("%s()"), *RegisterHelperName));
+	Context.AddLine(TEXT("{"));
+	Context.IncreaseIndent();
+
+	Context.AddLine(FString::Printf(
+		TEXT("FConvertedBlueprintsDependencies::Get().RegisterClass(TEXT(\"%s\"), &%s::__StaticDependenciesAssets);")
+		, *OriginalClass->GetOutermost()->GetPathName()
+		, *CppClassName));
 
 	Context.DecreaseIndent();
 	Context.AddLine(TEXT("}"));
+
+	Context.AddLine(FString::Printf(TEXT("static %s Instance;"), *RegisterHelperName));
+
+	Context.DecreaseIndent();
+	Context.AddLine(TEXT("};"));
+
+	Context.AddLine(FString::Printf(TEXT("%s %s::Instance;"), *RegisterHelperName, *RegisterHelperName));
+}
+
+
+void FEmitDefaultValueHelper::FillCommonUsedAssets(FEmitterLocalContext& Context, TSharedPtr<FGatherConvertedClassDependencies> ParentDependencies)
+{
+	ensure(0 == Context.UsedObjectInCurrentClass.Num());
+	for (auto LocAsset : Context.Dependencies.Assets)
+	{
+		if (!ParentDependencies.IsValid() || !ParentDependencies->Assets.Contains(LocAsset))
+		{
+			Context.UsedObjectInCurrentClass.Add(LocAsset);
+		}
+	}
 }
 
 void FEmitDefaultValueHelper::GenerateCustomDynamicClassInitialization(FEmitterLocalContext& Context, TSharedPtr<FGatherConvertedClassDependencies> ParentDependencies)
@@ -907,15 +1043,7 @@ void FEmitDefaultValueHelper::GenerateCustomDynamicClassInitialization(FEmitterL
 		});
 	}
 
-	Context.AddLine(TEXT("__CustomDynamicClassInitializationUsedAssets(InDynamicClass);"));
-
-	for (auto LocAsset : Context.Dependencies.Assets)
-	{
-		if (!ParentDependencies.IsValid() || !ParentDependencies->Assets.Contains(LocAsset))
-		{
-			Context.UsedObjectInCurrentClass.Add(LocAsset);
-		}
-	}
+	Context.AddLine(TEXT("FConvertedBlueprintsDependencies::FillUsedAssetsInDynamicClass(InDynamicClass, &__StaticDependencies_DirectlyUsedAssets);"));
 
 	auto CreateAndInitializeClassSubobjects = [&](bool bCreate, bool bInitilize)
 	{
@@ -956,12 +1084,6 @@ void FEmitDefaultValueHelper::GenerateCustomDynamicClassInitialization(FEmitterL
 void FEmitDefaultValueHelper::GenerateConstructor(FEmitterLocalContext& Context)
 {
 	UBlueprintGeneratedClass* BPGC = CastChecked<UBlueprintGeneratedClass>(Context.GetCurrentlyGeneratedClass());
-	UBlueprintGeneratedClass* ParentBPGC = Cast<UBlueprintGeneratedClass>(BPGC->GetSuperClass());
-	TSharedPtr<FGatherConvertedClassDependencies> ParentDependencies(ParentBPGC ? new FGatherConvertedClassDependencies(ParentBPGC) : nullptr);
-
-	GenerateCustomDynamicClassInitialization(Context, ParentDependencies);
-
-	
 	const FString CppClassName = FEmitHelper::GetCppName(BPGC);
 
 	UClass* SuperClass = BPGC->GetSuperClass();
@@ -1138,8 +1260,6 @@ void FEmitDefaultValueHelper::GenerateConstructor(FEmitterLocalContext& Context)
 	Context.DecreaseIndent();
 	Context.AddLine(TEXT("}"));
 
-	FDependenciesHelper::AddStaticFunctionsForDependencies(Context, ParentDependencies);
-
 	FBackendHelperUMG::EmitWidgetInitializationFunctions(Context);
 }
 
diff --git a/Engine/Source/Developer/BlueprintCompilerCppBackend/Public/BlueprintCompilerCppBackendGatherDependencies.h b/Engine/Source/Developer/BlueprintCompilerCppBackend/Public/BlueprintCompilerCppBackendGatherDependencies.h
--- a/Engine/Source/Developer/BlueprintCompilerCppBackend/Public/BlueprintCompilerCppBackendGatherDependencies.h
+++ b/Engine/Source/Developer/BlueprintCompilerCppBackend/Public/BlueprintCompilerCppBackendGatherDependencies.h
@@ -34,6 +34,8 @@ public:
 
 	UClass* FindOriginalClass(const UClass* InClass) const;
 
+	UClass* GetFirstNativeOrConvertedClass(UClass* InClass, bool bExcludeBPDataOnly = false) const;
+
 public:
 	bool WillClassBeConverted(const UBlueprintGeneratedClass* InClass) const;
 
diff --git a/Engine/Source/Developer/BlueprintProfiler/Private/BlueprintProfiler.cpp b/Engine/Source/Developer/BlueprintProfiler/Private/BlueprintProfiler.cpp
--- a/Engine/Source/Developer/BlueprintProfiler/Private/BlueprintProfiler.cpp
+++ b/Engine/Source/Developer/BlueprintProfiler/Private/BlueprintProfiler.cpp
@@ -20,19 +20,7 @@ FBlueprintProfiler::FBlueprintProfiler()
 	: bProfilerActive(false)
 	, bProfilingCaptureActive(false)
 	, bPIEActive(false)
-	, GraphNodeHeatMapDisplayMode(EBlueprintProfilerHeatMapDisplayMode::None)
-	, WireHeatMapDisplayMode(EBlueprintProfilerHeatMapDisplayMode::None)
 {
-	// Update stat watermarks from editor settings, the callbacks should keep them in sync after this.
-	if (const UEditorExperimentalSettings* EditorSettings = GetDefault<UEditorExperimentalSettings>())
-	{
-		FScriptPerfData::EnableRecentSampleBias(EditorSettings->bEnableBlueprintProfilerRecentSampleBias);
-		FScriptPerfData::SetRecentSampleBias(EditorSettings->BlueprintProfilerRecentSampleBias);
-		FScriptPerfData::SetEventPerformanceThreshold(EditorSettings->BlueprintProfilerEventThreshold);
-		FScriptPerfData::SetExclusivePerformanceThreshold(EditorSettings->BlueprintProfilerExclNodeThreshold);
-		FScriptPerfData::SetInclusivePerformanceThreshold(EditorSettings->BlueprintProfilerInclNodeThreshold);
-		FScriptPerfData::SetMaxPerformanceThreshold(EditorSettings->BlueprintProfilerMaxNodeThreshold);
-	}
 }
 
 FBlueprintProfiler::~FBlueprintProfiler()
@@ -66,7 +54,7 @@ void FBlueprintProfiler::ToggleProfilingCapture()
 	FBlueprintCoreDelegates::OnToggleScriptProfiler.Broadcast(bProfilerActive);
 }
 
-void FBlueprintProfiler::InstrumentEvent(const EScriptInstrumentationEvent& Event)
+void FBlueprintProfiler::InstrumentEvent(const FScriptInstrumentationSignal& Event)
 {
 	#if WITH_EDITOR
 	SCOPE_CYCLE_COUNTER(STAT_ProfilerInstrumentationCost);
@@ -76,22 +64,9 @@ void FBlueprintProfiler::InstrumentEvent(const EScriptInstrumentationEvent& Even
 	#endif
 	{
 		// Handle context switching events
-		CaptureContext.UpdateContext(Event.GetContextObject(), InstrumentationEventQueue);
+		CaptureContext.UpdateContext(Event, InstrumentationEventQueue);
 		// Add instrumented event
-		int32 ScriptCodeOffset = -1;
-		FName CurrentFunctionName = Event.GetEventName();
-		if (Event.IsStackFrameValid())
-		{
-			const FFrame& StackFrame = Event.GetStackFrame();
-			ScriptCodeOffset = StackFrame.Code - StackFrame.Node->Script.GetData() - 1;
-			CurrentFunctionName = StackFrame.Node->GetFName();
-		}
-		// Handle latent suspends - set the script offset to the LinkId of the most recent latent action.
-		if (Event.GetType() == EScriptInstrumentation::ResumeEvent)
-		{
-			ScriptCodeOffset = Event.GetLatentLinkId();
-		}
-		InstrumentationEventQueue.Add(FScriptInstrumentedEvent(Event.GetType(), CurrentFunctionName, ScriptCodeOffset));
+		InstrumentationEventQueue.Add(FScriptInstrumentedEvent(Event.GetType(), Event.GetFunctionClassScope()->GetFName(), Event.GetFunctionName(), Event.GetScriptCodeOffset()));
 		// Reset context on event end
 		if (Event.GetType() == EScriptInstrumentation::Stop)
 		{
@@ -102,6 +77,14 @@ void FBlueprintProfiler::InstrumentEvent(const EScriptInstrumentationEvent& Even
 
 #if WITH_EDITOR
 
+void FBlueprintProfiler::AddInstrumentedBlueprint(UBlueprint* InstrumentedBlueprint)
+{
+	if (InstrumentedBlueprint && InstrumentedBlueprint->GeneratedClass && InstrumentedBlueprint->GeneratedClass->HasInstrumentation())
+	{
+		GetBlueprintContext(InstrumentedBlueprint->GeneratedClass->GetPathName());
+	}
+}
+
 TSharedPtr<FBlueprintExecutionContext> FBlueprintProfiler::GetBlueprintContext(const FString& BlueprintClassPath)
 {
 	SCOPE_CYCLE_COUNTER(STAT_BlueprintLookupCost);
@@ -144,6 +127,19 @@ TSharedPtr<FScriptExecutionNode> FBlueprintProfiler::GetProfilerDataForNode(cons
 	return Result;
 }
 
+TSharedPtr<FScriptExecutionBlueprint> FBlueprintProfiler::GetProfilerDataForBlueprint(const UBlueprint* Blueprint)
+{
+	TSharedPtr<FScriptExecutionBlueprint> Result;
+	if (Blueprint)
+	{
+		if (TSharedPtr<FBlueprintExecutionContext>* BlueprintContext = PathToBlueprintContext.Find(Blueprint->GeneratedClass->GetPathName()))
+		{
+			Result = (*BlueprintContext)->GetBlueprintExecNode();
+		}
+	}
+	return Result;
+}
+
 bool FBlueprintProfiler::HasDataForInstance(const UObject* Instance) const
 {
 	bool bHasData = false;
@@ -185,6 +181,7 @@ void FBlueprintProfiler::ProcessEventProfilingData()
 		bool IsRangeValid() const { return StopIdx > StartIdx; }
 
 		TSharedPtr<FBlueprintExecutionContext> BlueprintContext;
+		FName ClassScope;
 		FName InstanceName;
 		int32 StartIdx;
 		int32 StopIdx;
@@ -207,10 +204,11 @@ void FBlueprintProfiler::ProcessEventProfilingData()
 					if (InstanceEvent.IsNewInstance() && Event.IsEvent())
 					{
 						FEventRange NewEventRange;
-						NewEventRange.BlueprintContext = GetBlueprintContext(CurrEvent.GetObjectPath());
+						NewEventRange.BlueprintContext = GetBlueprintContext(CurrEvent.GetBlueprintClassPath());
 						if (NewEventRange.BlueprintContext.IsValid())
 						{
-							NewEventRange.InstanceName = NewEventRange.BlueprintContext->MapBlueprintInstance(InstanceEvent.GetObjectPath());
+							NewEventRange.InstanceName = NewEventRange.BlueprintContext->MapBlueprintInstance(InstanceEvent.GetInstancePath());
+							NewEventRange.ClassScope = Event.GetFunctionClassScopeName();
 							NewEventRange.StartIdx = EventIdx;
 							ScriptEventRanges.Push(NewEventRange);
 						}
@@ -221,15 +219,78 @@ void FBlueprintProfiler::ProcessEventProfilingData()
 				}
 				break;
 			}
+			case EScriptInstrumentation::ClassScope:
+			{
+				// Class scope switch handling.
+				if (ScriptEventRanges.Num() > 0)
+				{
+					// Check against the last class scope first before adding another range.
+					FEventRange& LastEventRange = ScriptEventRanges.Last();
+					const FName NewClassScope = InstrumentationEventQueue[EventIdx].GetFunctionClassScopeName();
+					if (LastEventRange.ClassScope != NewClassScope)
+					{
+						// Override the class scope signal to be an event.
+						InstrumentationEventQueue[EventIdx].OverrideType(EScriptInstrumentation::Event);
+						// Create new sub event event range to be processed as a sub event.
+						FEventRange NewClassScopeRange;
+						NewClassScopeRange.BlueprintContext = LastEventRange.BlueprintContext;
+						NewClassScopeRange.InstanceName = LastEventRange.InstanceName;
+						NewClassScopeRange.ClassScope = NewClassScope;
+						NewClassScopeRange.StartIdx = EventIdx;
+						ScriptEventRanges.Push(NewClassScopeRange);
+						// find new class subset
+						double StartTime = MAX_dbl;
+						double StopTime = MIN_dbl;
+						int32 SubClassIdx = EventIdx+1;
+						for (; SubClassIdx < InstrumentationEventQueue.Num(); ++SubClassIdx)
+						{
+							// Check for end of scope
+							if (InstrumentationEventQueue[SubClassIdx].GetFunctionClassScopeName() != NewClassScope)
+							{
+								break;
+							}
+							// Update new event timings
+							const double EventTime = InstrumentationEventQueue[SubClassIdx].GetTime();
+							StartTime = FMath::Min<double>(StartTime, EventTime);
+							StopTime = FMath::Max<double>(StopTime, EventTime);
+						}
+						// Set event timings
+						InstrumentationEventQueue[EventIdx].OverrideTime(StartTime);
+						// Insert Stop marker
+						FScriptInstrumentedEvent StopMarker(EScriptInstrumentation::Stop, NewClassScope, InstrumentationEventQueue[EventIdx].GetFunctionName(), INDEX_NONE);
+						StopMarker.OverrideTime(StopTime);
+						InstrumentationEventQueue.Insert(StopMarker, SubClassIdx);
+						// Note: we could skip the class scope signals and jump straight to processing here - it might cause problems though.
+					}
+					else
+					{
+						// we have processed a scope switch and are now continuing with the previous scope.
+						// just remove the scope signal and adjust the processing offset ( so we process the next valid signal ).
+						InstrumentationEventQueue.RemoveAt(EventIdx, 1, false);
+						EventIdx--;
+					}
+				}
+				break;
+			}
 			case EScriptInstrumentation::Event:
 			{
 				// Nested events such as calls from event dispatchers.
 				if (ScriptEventRanges.Num() > 0)
 				{
 					FEventRange& LastEventRange = ScriptEventRanges.Last();
+					FName ActiveInstance = LastEventRange.InstanceName;
+					// Check for a different instance, update the active instance and remove the signal.
+					const int32 InstanceIdx = EventIdx - 1;
+					if (EventIdx > 0 && InstrumentationEventQueue[InstanceIdx].GetType() == EScriptInstrumentation::Instance)
+					{
+						ActiveInstance = LastEventRange.BlueprintContext->MapBlueprintInstance(InstrumentationEventQueue[InstanceIdx].GetInstancePath());
+						InstrumentationEventQueue.RemoveAt(InstanceIdx, 1, false);
+						EventIdx--;
+					}
 					FEventRange NewEventRange;
 					NewEventRange.BlueprintContext = LastEventRange.BlueprintContext;
-					NewEventRange.InstanceName = LastEventRange.InstanceName;
+					NewEventRange.InstanceName = ActiveInstance;
+					NewEventRange.ClassScope = CurrEvent.GetFunctionClassScopeName();
 					NewEventRange.StartIdx = EventIdx;
 					ScriptEventRanges.Push(NewEventRange);
 				}
diff --git a/Engine/Source/Developer/BlueprintProfiler/Private/BlueprintProfilerConnectionDrawingPolicy.cpp b/Engine/Source/Developer/BlueprintProfiler/Private/BlueprintProfilerConnectionDrawingPolicy.cpp
--- a/Engine/Source/Developer/BlueprintProfiler/Private/BlueprintProfilerConnectionDrawingPolicy.cpp
+++ b/Engine/Source/Developer/BlueprintProfiler/Private/BlueprintProfilerConnectionDrawingPolicy.cpp
@@ -17,10 +17,10 @@ FConnectionDrawingPolicy* FBlueprintProfilerPinConnectionFactory::CreateConnecti
 		UBlueprintGeneratedClass* BPGC = Blueprint ? Cast<UBlueprintGeneratedClass>(Blueprint->GeneratedClass) : nullptr;
 		if (BPGC && BPGC->HasInstrumentation())
 		{
-			IBlueprintProfilerInterface& ProfilerModule = FModuleManager::LoadModuleChecked<IBlueprintProfilerInterface>("BlueprintProfiler");
-			EBlueprintProfilerHeatMapDisplayMode::Type WireHeatMode = ProfilerModule.GetWireHeatMapDisplayMode();
+			EBlueprintProfilerHeatMapDisplayMode WireHeatMode = GetDefault<UBlueprintProfilerSettings>()->WireHeatMapDisplayMode;
 			if (WireHeatMode != EBlueprintProfilerHeatMapDisplayMode::None)
 			{
+				IBlueprintProfilerInterface& ProfilerModule = FModuleManager::LoadModuleChecked<IBlueprintProfilerInterface>("BlueprintProfiler");
 				TSharedPtr<FBlueprintExecutionContext> BlueprintContext = ProfilerModule.GetBlueprintContext(BPGC->GetPathName());
 				if (BlueprintContext.IsValid())
 				{
@@ -35,7 +35,7 @@ FConnectionDrawingPolicy* FBlueprintProfilerPinConnectionFactory::CreateConnecti
 /////////////////////////////////////////////////////
 // FBlueprintProfilerConnectionDrawingPolicy
 
-FBlueprintProfilerConnectionDrawingPolicy::FBlueprintProfilerConnectionDrawingPolicy(int32 InBackLayerID, int32 InFrontLayerID, float ZoomFactor, const FSlateRect& InClippingRect, FSlateWindowElementList& InDrawElements, EBlueprintProfilerHeatMapDisplayMode::Type InHeatmapType, UEdGraph* InGraphObj)
+FBlueprintProfilerConnectionDrawingPolicy::FBlueprintProfilerConnectionDrawingPolicy(int32 InBackLayerID, int32 InFrontLayerID, float ZoomFactor, const FSlateRect& InClippingRect, FSlateWindowElementList& InDrawElements, EBlueprintProfilerHeatMapDisplayMode InHeatmapType, UEdGraph* InGraphObj)
 	: FKismetConnectionDrawingPolicy(InBackLayerID, InFrontLayerID, ZoomFactor, InClippingRect, InDrawElements, InGraphObj)
 	, WireHeatMode(InHeatmapType)
 	, GraphReference(InGraphObj)
@@ -58,14 +58,16 @@ void FBlueprintProfilerConnectionDrawingPolicy::BuildExecutionRoadmap()
 
 	// Grab relevant profiler interfaces
 	UBlueprintGeneratedClass* TargetClass = Cast<UBlueprintGeneratedClass>(TargetBP->GeneratedClass);
-	FBlueprintDebugData& DebugData = TargetClass->GetDebugData();
 	IBlueprintProfilerInterface& ProfilerModule = FModuleManager::LoadModuleChecked<IBlueprintProfilerInterface>("BlueprintProfiler");
 	TSharedPtr<FBlueprintExecutionContext> BlueprintContext = ProfilerModule.GetBlueprintContext(TargetClass->GetPathName());
 	const FName InstanceName = BlueprintContext->MapBlueprintInstance(ActiveObject->GetPathName());
-	const FName GraphName = GraphObj->GetFName();
-	TSharedPtr<FBlueprintFunctionContext> FunctionContext = BlueprintContext->GetFunctionContext(GraphName);
-	const FName FunctionName = FunctionContext->GetFunctionName();
-	UFunction* FunctionPtr = TargetClass->FindFunctionByName(FunctionName);
+	TSharedPtr<FBlueprintFunctionContext> FunctionContext = BlueprintContext->GetFunctionContextFromGraph(GraphObj);
+
+	if (!FunctionContext.IsValid())
+	{
+		// No point proceeding without a valid function context.
+		return;
+	}
 
 	// Redirect the target Blueprint when debugging with a macro graph visible
 	if (TargetBP->BlueprintType == BPTYPE_MacroLibrary)
@@ -73,36 +75,38 @@ void FBlueprintProfilerConnectionDrawingPolicy::BuildExecutionRoadmap()
 		TargetBP = Cast<UBlueprint>(ActiveObject->GetClass()->ClassGeneratedBy);
 	}
 
-	TArray<UEdGraphNode*> SequentialNodesInGraph;
 	TArray<double> SequentialNodeTimes;
 	TArray<UEdGraphPin*> SequentialExecPinsInGraph;
 	TArray<FTracePath> SequentialTracePathsInGraph;
+	TArray<TWeakPtr<FScriptExecutionNode>> SequentialProfilerNodesInGraph;
+	const FName ScopedFunctionName = FunctionContext->GetFunctionName();
 
 	const TSimpleRingBuffer<FBlueprintExecutionTrace>& TraceHistory = BlueprintContext->GetTraceHistory();
 	for (int32 i = 0; i < TraceHistory.Num(); ++i)
 	{
 		const FBlueprintExecutionTrace& Sample = TraceHistory(i);
-		if (InstanceName == Sample.InstanceName && FunctionName == Sample.FunctionName)
+		if (InstanceName == Sample.InstanceName && ScopedFunctionName == Sample.FunctionName)
 		{
-			TSharedPtr<FScriptExecutionNode> ProfilerNode = FunctionContext->GetProfilerDataForNode(Sample.NodeName);
-			check(ProfilerNode.IsValid());
-			UEdGraphPin* Pin = DebugData.FindSourcePinFromCodeLocation(FunctionPtr, Sample.Offset);
-
-			if (Pin && Pin->PinType.PinCategory == UEdGraphSchema_K2::PC_Exec)
+			if (Sample.ProfilerNode.IsValid())
 			{
-				if (const UEdGraphNode* Node = Cast<UEdGraphNode>(ProfilerNode->GetObservedObject()))
+				const UEdGraphPin* Pin = FunctionContext->GetPinFromCodeLocation(Sample.Offset);
+
+				if (Pin && Pin->PinType.PinCategory == UEdGraphSchema_K2::PC_Exec)
 				{
-					SequentialNodesInGraph.Add(const_cast<UEdGraphNode*>(Node));
-					SequentialNodeTimes.Add(Sample.ObservationTime);
-					SequentialExecPinsInGraph.Add(Pin);
-					SequentialTracePathsInGraph.Add(Sample.TracePath);
+					if (const UEdGraphNode* Node = Pin->GetOwningNode())
+					{
+						SequentialProfilerNodesInGraph.Add(Sample.ProfilerNode);
+						SequentialNodeTimes.Add(Sample.ObservationTime);
+						SequentialExecPinsInGraph.Add(const_cast<UEdGraphPin*>(Pin));
+						SequentialTracePathsInGraph.Add(Sample.TracePath);
+					}
 				}
 			}
 		}
 	}
 	// Run thru and apply bonus time
 	const float InvNumNodes = 1.0f / (float)SequentialNodeTimes.Num();
-	for (int32 i = 0; i < SequentialNodesInGraph.Num(); ++i)
+	for (int32 i = 0; i < SequentialProfilerNodesInGraph.Num(); ++i)
 	{
 		double& ObservationTime = SequentialNodeTimes[i];
 
@@ -136,76 +140,65 @@ void FBlueprintProfilerConnectionDrawingPolicy::BuildExecutionRoadmap()
 			}
 			if (bLinked)
 			{
-				UEdGraphNode* CurNode = OutputPin->GetOwningNode();
-				UEdGraphNode* NextNode = InputPin->GetOwningNode();
-				if (CurNode != NextNode)
-				{
-					TSharedPtr<FScriptExecutionNode> CurExecNode = BlueprintContext->GetProfilerDataForPin(OutputPin);
-					TSharedPtr<FScriptExecutionNode> NextExecNode = BlueprintContext->GetProfilerDataForNode(NextNode);
+				TSharedPtr<FScriptExecutionNode> CurExecNode = SequentialProfilerNodesInGraph[OutputPinIndex].Pin();
+				TSharedPtr<FScriptExecutionNode> NextExecNode = SequentialProfilerNodesInGraph[i].Pin();
 
-					if (CurExecNode.IsValid() && NextExecNode.IsValid())
+				if (CurExecNode.IsValid() && NextExecNode.IsValid() && CurExecNode != NextExecNode)
+				{
+					double NextNodeTime = SequentialNodeTimes[i];
+					UEdGraphNode* NextNode = const_cast<UEdGraphNode*>(NextExecNode->GetTypedObservedObject<UEdGraphNode>());
+					FExecPairingMap& ExecPaths			= PredecessorPins.FindOrAdd(NextNode);
+					FTimePair& ExecTiming				= ExecPaths.FindOrAdd(OutputPin);
+					FProfilerPairingMap& ProfilerPaths	= ProfilerPins.FindOrAdd(NextNode);
+					FProfilerPair& ProfilerData			= ProfilerPaths.FindOrAdd(OutputPin);
+					// make sure that if we've already visited this exec-pin (like 
+					// in a for-loop or something), that we're replacing it with a 
+					// more recent execution time
+					//
+					// @TODO I don't see when this wouldn't be the case
+					if (ExecTiming.ThisExecTime < NextNodeTime)
 					{
-						double NextNodeTime = SequentialNodeTimes[i];
-
-						FExecPairingMap& ExecPaths			= PredecessorPins.FindOrAdd(NextNode);
-						FTimePair& ExecTiming				= ExecPaths.FindOrAdd(OutputPin);
-						FProfilerPairingMap& ProfilerPaths	= ProfilerPins.FindOrAdd(NextNode);
-						FProfilerPair& ProfilerData			= ProfilerPaths.FindOrAdd(OutputPin);
-						// make sure that if we've already visited this exec-pin (like 
-						// in a for-loop or something), that we're replacing it with a 
-						// more recent execution time
-						//
-						// @TODO I don't see when this wouldn't be the case
-						if (ExecTiming.ThisExecTime < NextNodeTime)
+						double CurNodeTime = SequentialNodeTimes[OutputPinIndex];
+						ExecTiming.ThisExecTime = NextNodeTime;
+						ExecTiming.PredExecTime = CurNodeTime;
+						const bool bExecPin = CurExecNode->HasFlags(EScriptExecutionNodeFlags::ExecPin);
+						const FTracePath& CurTracePath = bExecPin ? SequentialTracePathsInGraph[i] : SequentialTracePathsInGraph[OutputPinIndex];
+						const FTracePath& NextTracePath = SequentialTracePathsInGraph[i];
+						TSharedPtr<FScriptPerfData> CurPerfData = CurExecNode->GetPerfDataByInstanceAndTracePath(InstanceName, CurTracePath);
+						TSharedPtr<FScriptPerfData> NextPerfData = NextExecNode->GetPerfDataByInstanceAndTracePath(InstanceName, NextTracePath);
+						ProfilerData.bPureNode = CurExecNode->IsPureNode();
+						switch(WireHeatMode)
 						{
-							double CurNodeTime = SequentialNodeTimes[OutputPinIndex];
-							ExecTiming.ThisExecTime = NextNodeTime;
-							ExecTiming.PredExecTime = CurNodeTime;
-							const bool bExecPin = CurExecNode->HasFlags(EScriptExecutionNodeFlags::ExecPin);
-							const FTracePath& CurTracePath = bExecPin ? SequentialTracePathsInGraph[i] : SequentialTracePathsInGraph[OutputPinIndex];
-							const FTracePath& NextTracePath = SequentialTracePathsInGraph[i];
-							TSharedPtr<FScriptPerfData> CurPerfData = CurExecNode->GetPerfDataByInstanceAndTracePath(InstanceName, CurTracePath);
-							TSharedPtr<FScriptPerfData> NextPerfData = NextExecNode->GetPerfDataByInstanceAndTracePath(InstanceName, NextTracePath);
-							ProfilerData.bPureNode = CurExecNode->IsPureNode();
-							switch(WireHeatMode)
+							case EBlueprintProfilerHeatMapDisplayMode::Average:
+							{
+								ProfilerData.PredPerfData = CurPerfData->GetAverageHeatLevel();
+								ProfilerData.ThisPerfData = NextPerfData->GetAverageHeatLevel();
+								break;
+							}
+							case EBlueprintProfilerHeatMapDisplayMode::Inclusive:
+							case EBlueprintProfilerHeatMapDisplayMode::PinToPin:
+							{
+								ProfilerData.PredPerfData = CurPerfData->GetInclusiveHeatLevel();
+								ProfilerData.ThisPerfData = NextPerfData->GetInclusiveHeatLevel();
+								break;
+							}
+							case EBlueprintProfilerHeatMapDisplayMode::MaxTiming:
+							{
+								ProfilerData.PredPerfData = CurPerfData->GetMaxTimeHeatLevel();
+								ProfilerData.ThisPerfData = NextPerfData->GetMaxTimeHeatLevel();
+								break;
+							}
+							case EBlueprintProfilerHeatMapDisplayMode::Total:
+							{
+								ProfilerData.PredPerfData = CurPerfData->GetTotalHeatLevel();
+								ProfilerData.ThisPerfData = NextPerfData->GetTotalHeatLevel();
+								break;
+							}
+							case EBlueprintProfilerHeatMapDisplayMode::HottestPath:
 							{
-								case EBlueprintProfilerHeatMapDisplayMode::Exclusive:
-								{
-									ProfilerData.PredPerfData = CurPerfData->GetNodeHeatLevel();
-									ProfilerData.ThisPerfData = NextPerfData->GetNodeHeatLevel();
-									break;
-								}
-								case EBlueprintProfilerHeatMapDisplayMode::Inclusive:
-								case EBlueprintProfilerHeatMapDisplayMode::PinToPin:
-								{
-									ProfilerData.PredPerfData = CurPerfData->GetInclusiveHeatLevel();
-									ProfilerData.ThisPerfData = NextPerfData->GetInclusiveHeatLevel();
-									break;
-								}
-								case EBlueprintProfilerHeatMapDisplayMode::MaxTiming:
-								{
-									ProfilerData.PredPerfData = CurPerfData->GetMaxTimeHeatLevel();
-									ProfilerData.ThisPerfData = NextPerfData->GetMaxTimeHeatLevel();
-									break;
-								}
-								case EBlueprintProfilerHeatMapDisplayMode::Total:
-								{
-									ProfilerData.PredPerfData = CurPerfData->GetTotalHeatLevel();
-									ProfilerData.ThisPerfData = NextPerfData->GetTotalHeatLevel();
-									break;
-								}
-								case EBlueprintProfilerHeatMapDisplayMode::HottestPath:
-								{
-									ProfilerData.PredPerfData = CurPerfData->GetHottestPathHeatLevel();
-									ProfilerData.ThisPerfData = NextPerfData->GetHottestPathHeatLevel();
-									break;
-								}
-								case EBlueprintProfilerHeatMapDisplayMode::HottestEndpoint:
-								{
-									ProfilerData.PredPerfData = CurPerfData->GetHottestEndpointHeatLevel();
-									ProfilerData.ThisPerfData = NextPerfData->GetHottestEndpointHeatLevel();
-									break;
-								}
+								ProfilerData.PredPerfData = CurPerfData->GetHottestPathHeatLevel();
+								ProfilerData.ThisPerfData = NextPerfData->GetHottestPathHeatLevel();
+								break;
 							}
 						}
 					}
diff --git a/Engine/Source/Developer/BlueprintProfiler/Private/BlueprintProfilerConnectionDrawingPolicy.h b/Engine/Source/Developer/BlueprintProfiler/Private/BlueprintProfilerConnectionDrawingPolicy.h
--- a/Engine/Source/Developer/BlueprintProfiler/Private/BlueprintProfilerConnectionDrawingPolicy.h
+++ b/Engine/Source/Developer/BlueprintProfiler/Private/BlueprintProfilerConnectionDrawingPolicy.h
@@ -3,7 +3,7 @@
 #pragma once
 
 #include "Editor/GraphEditor/Public/BlueprintConnectionDrawingPolicy.h"
-#include "BlueprintProfilerModule.h"
+#include "Editor/Kismet/Public/Profiler/BlueprintProfilerSettings.h"
 
 /////////////////////////////////////////////////////
 // FBlueprintProfilerPinConnectionFactory
@@ -64,7 +64,7 @@ protected:
 
 public:
 
-	FBlueprintProfilerConnectionDrawingPolicy(int32 InBackLayerID, int32 InFrontLayerID, float ZoomFactor, const FSlateRect& InClippingRect, FSlateWindowElementList& InDrawElements, EBlueprintProfilerHeatMapDisplayMode::Type InHeatmapType, UEdGraph* InGraphObj);
+	FBlueprintProfilerConnectionDrawingPolicy(int32 InBackLayerID, int32 InFrontLayerID, float ZoomFactor, const FSlateRect& InClippingRect, FSlateWindowElementList& InDrawElements, EBlueprintProfilerHeatMapDisplayMode InHeatmapType, UEdGraph* InGraphObj);
 
 	// FKismetConnectionDrawingPolicy
 	virtual void BuildExecutionRoadmap() override;
@@ -88,7 +88,7 @@ private:
 private:
 
 	/** The current wire heat map mode */
-	EBlueprintProfilerHeatMapDisplayMode::Type WireHeatMode;
+	EBlueprintProfilerHeatMapDisplayMode WireHeatMode;
 	/** The current graph */
 	TWeakObjectPtr<UEdGraph> GraphReference;
 
diff --git a/Engine/Source/Developer/BlueprintProfiler/Private/ScriptInstrumentationCapture.cpp b/Engine/Source/Developer/BlueprintProfiler/Private/ScriptInstrumentationCapture.cpp
--- a/Engine/Source/Developer/BlueprintProfiler/Private/ScriptInstrumentationCapture.cpp
+++ b/Engine/Source/Developer/BlueprintProfiler/Private/ScriptInstrumentationCapture.cpp
@@ -13,31 +13,49 @@ void FScriptInstrumentedEvent::SetData(EScriptInstrumentation::Type InEventType,
 	Time = FPlatformTime::Seconds();
 }
 
+FName FScriptInstrumentedEvent::GetScopedFunctionName() const
+{
+	return FName(*FString::Printf(TEXT("%s::%s"), *FunctionClassScopeName.ToString(), *FunctionName.ToString()));
+}
+
 //////////////////////////////////////////////////////////////////////////
 // FInstrumentationCaptureContext
 
-void FInstrumentationCaptureContext::UpdateContext(const UObject* InContextObject, TArray<FScriptInstrumentedEvent>& InstrumentationQueue)
+void FInstrumentationCaptureContext::UpdateContext(const FScriptInstrumentationSignal& InstrumentSignal, TArray<FScriptInstrumentedEvent>& InstrumentationQueue)
 {
 	// Handle instance context switch
-	if (ContextObject != InContextObject && InContextObject)
+	if (const UObject* NewContextObject = InstrumentSignal.GetContextObject())
 	{
-		// Handle class context switch
-		if (const UClass* NewClass = InContextObject->GetClass())
+		if (NewContextObject != ContextObject)
+		{
+			// Handle class context switch
+			if (const UClass* NewClass = InstrumentSignal.GetClass())
+			{
+				if (NewClass != ContextClass)
+				{
+					InstrumentationQueue.Add(FScriptInstrumentedEvent(EScriptInstrumentation::Class, InstrumentSignal.GetFunctionClassScope()->GetFName(), NewClass->GetPathName()));
+					ContextClass = NewClass;
+				}
+			}
+			// Handle instance change
+			InstrumentationQueue.Add(FScriptInstrumentedEvent(EScriptInstrumentation::Instance, ContextClass->GetFName(), NewContextObject->GetPathName()));
+			ContextObject = NewContextObject;
+		}
+		else
 		{
-			if (NewClass != ContextClass)
+			// Handle function class scope changes
+			if (ContextFunctionClassScope != InstrumentSignal.GetFunctionClassScope())
 			{
-				InstrumentationQueue.Add(FScriptInstrumentedEvent(EScriptInstrumentation::Class, NewClass->GetPathName()));
-				ContextClass = NewClass;
+				ContextFunctionClassScope = InstrumentSignal.GetFunctionClassScope();
+				InstrumentationQueue.Add(FScriptInstrumentedEvent(EScriptInstrumentation::ClassScope, ContextFunctionClassScope->GetFName(), InstrumentSignal.GetFunctionName()));
 			}
 		}
-		// Handle instance change
-		InstrumentationQueue.Add(FScriptInstrumentedEvent(EScriptInstrumentation::Instance, InContextObject->GetPathName()));
-		ContextObject = InContextObject;
 	}
 }
 
 void FInstrumentationCaptureContext::ResetContext()
 {
 	ContextClass.Reset();
+	ContextFunctionClassScope.Reset();
 	ContextObject.Reset();
 }
\ No newline at end of file
