Copying //UE4/Dev-Blueprints to //UE4/Dev-Main (Source: //UE4/Dev-Blueprints @ 3235800)

==========================
MAJOR FEATURES + CHANGES
==========================

Change 3194900 on 2016/11/11 by [REDACTED_PERSON]

	Fixed issue where Reroute node pins weren't mirroring data properly.

Change 3195081 on 2016/11/11 by [REDACTED_PERSON]

	This [REDACTED_HANDLE] was addressed

Change 3196368 on 2016/11/14 by [REDACTED_PERSON]

	Results of FBlueprintNativeCodeGenModule::IsTargetedForReplacement are cashed - optimization (cooking time).

Change 3196369 on 2016/11/14 by [REDACTED_PERSON]

	CompileDisplaysBinaryBackend, CompileDisplaysTextBackend and bDisplaysLayout features (in [Kismet] in Engine.ini) are disabled in commandlets. They slow down BP compilation.

Change 3196398 on 2016/11/14 by [REDACTED_PERSON]

	Reworked the tracking of latent and expansion event tracking in the blueprint compiler for use with the blueprint profiler.
	#Proj BlueprintProfiler, KismetCompiler. BlueprintGraph, Engine

Change 3196410 on 2016/11/14 by [REDACTED_PERSON]

	Fixed crash in UK2Node_Knot::PropagatePinTypeFromInput

Change 3196852 on 2016/11/14 by [REDACTED_PERSON]

	Fixed static analysis warning.

Change 3196874 on 2016/11/14 by [REDACTED_PERSON]

	(the issue was already fixed, but it was reintroduced, when EDL support was added).

	ObjectImport->XObject is not filled prematurelly, so CreateExport is properly called each dynamic class.

Change 3197469 on 2016/11/14 by [REDACTED_PERSON]

	Fix for being able to make Sets and Maps of user defined structs that contained unhashable types (e.g. Rotator)

Change 3197703 on 2016/11/14 by [REDACTED_PERSON]

	Updated documentation comment to reflect current behavior

Change 3198167 on 2016/11/15 by [REDACTED_PERSON]

	Merged 3196582 from Dev-Core

	UE4 - Changed a check to a warning related to detaching linekrs twice. Seen in nativized BP version of platformer game.

Change 3198433 on 2016/11/15 by [REDACTED_PERSON]

	Fixed Copy/pasting variable nodes hides them from a reference search
	#UE-31606

Change 3198811 on 2016/11/15 by [REDACTED_PERSON]

	Fixed Knot node - it will use/propagate the type from input connection, if it's possible (intstead of the type from output connection).

Change 3198866 on 2016/11/15 by [REDACTED_PERSON]

	Fixed infinite loading of DynamicClass in EDL.

Change 3199045 on 2016/11/15 by [REDACTED_PERSON]

	[UE-27402] Fix attached Actor-based Blueprint instance root component relative transform values after reconstruction.

	change summary:
	- modified FActorComponentInstanceData's ctor to exclude relative transform properties when caching root component values.

Change 3200703 on 2016/11/16 by [REDACTED_PERSON]

	Marking the ease node explicitly as pure, which makes it so we can prune it from graphs when it is unused.

Change 3201115 on 2016/11/16 by [REDACTED_PERSON]

	Nativization + EDL: "Dynamic" objects are processed by FAsyncPackage::PostLoadDeferredObjects, so the EInternalObjectFlags::AsyncLoading flag is properly cleared.

Change 3201749 on 2016/11/17 by [REDACTED_PERSON]

	In EDL a package containig a dynamic class has PKG_CompiledIn flag (the same like without EDL).

Change 3202577 on 2016/11/17 by [REDACTED_PERSON]

	Accounting for a change in our intermediate node mapping - the old list no longer maps expanded nodes to macro instances (instead it maps to the corresponding node in the macro graph), so we had to use a new mapping meant for this.

Change 3204803 on 2016/11/18 by [REDACTED_PERSON]

	[UE-38607] Implicitly turn on Blueprint class nativization for dependencies.

	change summary:
	- added a UBlueprint::PostEditChangeProperty() override method to handle this.

Change 3204812 on 2016/11/18 by [REDACTED_PERSON]

	[UE-38580] Implicitly turn on the "nativize" project setting when enabling nativize for any Blueprint class.

	change summary:
	- modified UBlueprint::PostEditChangeProperty() to update project packaging settings if necessary

Change 3204818 on 2016/11/18 by [REDACTED_PERSON]

	[UE-38725] Interface class dependencies that are not enabled for nativization will now raise an error during nativized cooks.

	change summary:
	- modified FBlueprintNativeCodeGenModule::IsTargetedForReplacement() to check interface class dependencies in addition to parent class dependencies.

Change 3204963 on 2016/11/18 by [REDACTED_PERSON]

	Create a transaction when using UBlueprintFunctionNodeSpawner* directly

Change 3206510 on 2016/11/21 by [REDACTED_PERSON]

	Adding math-expression aliases for dot and cross functions ("dot" and "cross" respectively).

Change 3206547 on 2016/11/21 by [REDACTED_PERSON]

	Exposing GetReflectionVector() to Blueprints.

Change 3206658 on 2016/11/21 by [REDACTED_PERSON]

	Fix for compile error, digging out authorative class from trash class.

	Mirror of 3206638 from Odin

Change 3207579 on 2016/11/22 by [REDACTED_PERSON]

	No longer enforcing the requirement that game UFunctions have to have a category (making writing of game code easier).

Change 3207956 on 2016/11/22 by [REDACTED_PERSON]

	[UE-38690] Fix a regression in which nested scene component subobjects would no longer be registered after construction of an instance-added component in IWCE.

	change summary:
	- modified the IWCE path in SSCSEditor::AddNewComponent() to ensure that any components added as a result of instancing the newly-added component are also registered.
	- modified AActor::ExecuteConstruction() to ensure that non-scene, native nested component subobjects that might be created as a result of SCS execution are also registered (previously it was only considering non-scene components that were explicitly created by the SCS, or that inherited the creation method of the instance that created it).

Change 3208217 on 2016/11/22 by [REDACTED_PERSON]

	Modified fix (originally from Ryan.Rauschkolb, CL 3186023):

	Fixed unable to set struct variable name if name includes space

Change 3208347 on 2016/11/22 by [REDACTED_PERSON]

	Merging //UE4/Dev-Main to Dev-Blueprints (//UE4/Dev-Blueprints)

Change 3208688 on 2016/11/23 by [REDACTED_PERSON]

	Made a minor change that forces debugging references to the PIE world when the play in editor world is changed. This is intended to better handle mutliple game instance/world debugging scenarios.
	#Proj Engine, UnrealEd

Change 3208712 on 2016/11/23 by [REDACTED_PERSON]

	Improved handling of unwired transform struct terminal expression's in the blueprint VM compiler to remove an errant warning.
	#Proj KismetCompiler

Change 3209457 on 2016/11/23 by [REDACTED_PERSON]

	[UE-30479] Fix inability to edit the ISMC instance array on Actor instances when the ISMC is inherited from a Blueprint class.

	change summary:
	- added PPF_ForceTaggedSerialization as a means to override the CPF_SkipSerialization flag when explicit serialization of the property value is needed
	- modified UProperty::ShouldSerializeValue() to check for and handle the PPF_ForceTaggedSerialization flag when the CPF_SkipSerialization flag is present
	- modified UActorComponent::DetermineUCSModifiedProperties() to add the PPF_ForceTaggedSerialization flag to the custom FArchive impl
	- modified FActorComponentInstanceData::FActorComponentInstanceData() to add the PPF_ForceTaggedSerialization flag to the custom FObjectWriter impl
	- modified FActorComponentInstanceData::ApplyToComponent() to add the PPF_ForceTaggedSerialization flag to the custom FObjectReader impl

Change 3209758 on 2016/11/24 by [REDACTED_PERSON]

	Nativization: fixed error when a BP implements a native interface.
	FBlueprintNativeCodeGenModule::IsTargetedForReplacement will return "DontReplace" for native class.

Change 3210376 on 2016/11/25 by [REDACTED_PERSON]

	Fixed  FBlueprintNativeCodeGenModule::FindReplacedNameAndOuter
	Components in nativized BPCG SCS have replaced outer object and name while cooking.

Change 3210936 on 2016/11/28 by [REDACTED_PERSON]

	Minor revision to try and avoid a potentially expensive Contains() call when possible.

Change 3211527 on 2016/11/28 by [REDACTED_PERSON]

	Fixed map of names cooked in packages in nativized build.

Change 3211969 on 2016/11/28 by [REDACTED_PERSON]

	Merging //UE4/Dev-Main to Dev-Blueprints (//UE4/Dev-Blueprints)

Change 3212328 on 2016/11/28 by [REDACTED_PERSON]

	Enum, pointer and arithmetic specializations for THasGetTypeHash, as VC doesn't detect them properly.
	TIsEnum moved to its own header. Submitted on behalf of steve.robb

Change 3212398 on 2016/11/28 by [REDACTED_PERSON]

	Build fix, this function is part of another change

Change 3212442 on 2016/11/28 by [REDACTED_PERSON]

	UHT now supports structs in TMap and TSet, misc. fixes to PropertyStruct's PropertyFlags detecting whether the struct type is hashable (all HasGetTypeHash flags are now computed from THasGetTypeHash). Various fixes for generating TMap/TSet code from blueprints

Change 3212578 on 2016/11/28 by [REDACTED_PERSON]

	Rename RegisterClass to avoid collsion with RegistClass macro in generated code

Change 3213668 on 2016/11/29 by [REDACTED_PERSON]

	Fix for missing CDO when instatiating some subobjects in nativized BPs

Change 3213737 on 2016/11/29 by [REDACTED_PERSON]

	Added GetTypeHash implementation to UEnumProperty

Change 3215225 on 2016/11/30 by [REDACTED_PERSON]

	Bunch of changes required for nativized Orion to work with EDL.
	- ClientOnly, ServerOnly and EditorOnly assets are properly distinguished and handled
	- Introduced FCompilerNativizationOptions
	- fixed inproper references to BP instead of BPGC
	- fixed generated delegate name
	- hack for DefaultRootNode UE-39168
	- improved NativeCodeGenrationTool
	- various minor improvements

Change 3216363 on 2016/11/30 by [REDACTED_PERSON]

	Better fix for discrepency between UUserDefinedEnum::GetEnumText and UEnum::GetEnumText. Without meta data we could not look up display names, so I'm writing out the display names into a function in the BP cpp backend. This function could be generated by UHT if we wanted to correct this odd behavior for native enums

Change 3217168 on 2016/12/01 by [REDACTED_PERSON]

	Nativization:
	Fixed compilation warning C4458: declaration of 'CurrentState' hides class member
	Disabled warning C4996 (deprecated) in nativized code.

Change 3217320 on 2016/12/01 by [REDACTED_PERSON]

	[UE-38652] Selecting Blueprint assets for nativization is now integrated into the project's configuration.

	change summary:
	- added EProjectPackagingBlueprintNativizationMethod
	- deprecated the 'bNativizeBlueprintAssets' and 'bNativizeOnlySelectedBlueprints' flags in favor of a new 'BlueprintNativizationMethod' config property in UProjectPackagingSettings
	- added a new 'NativizeBlueprintAssets' config property to UProjectPackagingSettings to track/store the list of Blueprints to be nativized when the exclusive method (whitelist) is selected
	- added a PostInitProperties() override to UProjectPackagingSettings; implemented to migrate from the deprecated config properties to the new method enum type
	- updated FMainFrameActionCallbacks::PackageProject() to migrate to checking the enum type
	- updated FBlueprintNativeCodeGenModule::IsTargetedForReplacement() to migrate to checking the enum type
	- modified UProjectPackagingSettings::CanEditChange() to enable editing of the nativization whitelist only when the "exclusive" method is active
	- modified UProjectPackagingSettings::PostEditChangeProperty() to add a new case for handling changes to the whitelist; changes are propagated to any Blueprint assets that are loaded
	- added new Add/RemoveBlueprintAssetFromNativizationList() APIs to UProjectPackagingSettings for assisting with adding/removing Blueprint assets to/from the exclusive list (whitelist)
	- deprecated the 'bNativize' flag in UBlueprint in favor of a new transient 'bSelectedForNativization' flag that is no longer serialized; this now caches whether or not the asset is present in the whitelist in the project config
	- modified UBlueprint::Serialize() to both migrate from the deprecated flag to the project config/transient flag on load, as well as to propagate the value of the transient flag back to the project config on save. this means that if the user changes the value of the transient flag through the Details view, that change won't be reflected back to the project config until the Blueprint is actually saved (saving the value to the config rather than serializing to the asset)
	- modified UBlueprint::PostEditChangeProperty() to remove code that was previously updating the project configuration at edit time. this functionality has been relocated to Serialize() (save time) instead.

	notes:
	- also completes UE-38636 (task: consolidate config options into a single drop-down)

Change 3218124 on 2016/12/01 by [REDACTED_PERSON]

	 CPF_HasGetValueTypeHash flag was not set on native UEnumProperties

Change 3218168 on 2016/12/01 by [REDACTED_PERSON]

	Fix local var name that shadows a function param (CIS fix).

Change 3219117 on 2016/12/02 by [REDACTED_PERSON]

	Nativization:
	Fixed compilation warning C4458: when local function variable hides a class variable. Names of local variables in funvtions have prefix "bpfv__".

Change 3219201 on 2016/12/02 by [REDACTED_PERSON]

	Keeping the "Select All Input Nodes" option from infinitely recurssing by blocking on nodes it has already selected.

Change 3219247 on 2016/12/02 by [REDACTED_PERSON]

	Fixing CIS shadow variable warning from my last check in (CL 3219201).

Change 3219332 on 2016/12/02 by [REDACTED_PERSON]

	Workaround for UE-38231

Change 3219381 on 2016/12/02 by [REDACTED_PERSON]

	Accounting for cyclic compile issues in cast-node's validate function, making it check the authoratative class instead of the current type. Also down grading some of the warnings to notes (suggesting the users don't need the cast).

Change 3220224 on 2016/12/02 by [REDACTED_PERSON]

	Reduce font size for compact nodes

Change 3220476 on 2016/12/03 by [REDACTED_PERSON]

	Better fix for UE-35390
	Disabled deprecation warnings in nativized code.

Change 3221637 on 2016/12/05 by [REDACTED_PERSON]

	Nativization:
	Forced ImportCreation while InitialLoad for DynamicClasses.

Change 3222306 on 2016/12/05 by [REDACTED_PERSON]

	Support for default values of TMap and TSet local variables

Change 3222383 on 2016/12/05 by [REDACTED_PERSON]

	Fixed bug in Blueprint TMap function for getting values out of a TMap

Change 3222427 on 2016/12/05 by [REDACTED_PERSON]

	Preventing ChildActorTemplate fixups from occuring on component load, when they may instead be a placeholder object (a byproduct of deferred Blueprint loading - a guard against cyclic load problems).

Change 3222679 on 2016/12/05 by [REDACTED_PERSON]

	Remove unused code. Had no sideeffects, other than potential for leak when struct with non-trivial dtor was allocated here.

Change 3222719 on 2016/12/05 by [REDACTED_PERSON]

	Earlier detection of invalid native map/set properties. These generate a compile error if any TMap/TSet functions are used, but will slip by undetected if not used. Working on static assert to catch them.

Change 3224375 on 2016/12/06 by [REDACTED_PERSON]

	Add tags for testing of array diffing

Change 3224507 on 2016/12/07 by [REDACTED_PERSON]

	[UE-39055] Fix a crash caused by an object name collision that could occur when loading some older Blueprint assets.

	change summary:
	- added UInheritableComponentHandler::FixComponentTemplateName()
	- modified UInheritableComponentHandler::PostLoad() to check for and correct stale records that cause a collision with the corrected name
	- added a UInheritableComponentHandler::Serialize() override to ensure that UsingCustomVersion() is called for the asset containing the ICH (already happening in UBlueprint::Serialize(), but added for consistency with other usage)
	- modified USimpleConstructionScript::Serialize() to ensure that UsingCustomVersion() is called for the asset containing the SCS (same reason as above)

Change 3225572 on 2016/12/07 by [REDACTED_PERSON]

	Test assets for TSet/TMap testing. Includes new native class for testing containers. #rb none

Change 3225577 on 2016/12/07 by [REDACTED_PERSON]

	New test map for Array, TSet and Tmap testing.

Change 3226281 on 2016/12/07 by [REDACTED_PERSON]

	Container test asset, needs to be in p4 for diff tool tests.

Change 3226345 on 2016/12/07 by [REDACTED_PERSON]

	Another revision of test data

Change 3228496 on 2016/12/09 by [REDACTED_PERSON]

	This change adds extra information to component template arrays so that the component class can be determined in builds that strip out objects of certain class types such as the editor dedicated server build.
	#Proj KismetCompiler, BlueprintGraph, UnrealEd, Core, Engine, Kismet, BlueprintCompilerCppBackend

Change 3230120 on 2016/12/09 by [REDACTED_PERSON]

	Merging //UE4/Dev-Main to Dev-Blueprints (//UE4/Dev-Blueprints)

Change 3230700 on 2016/12/12 by [REDACTED_PERSON]

	Removing some array test properties from container test class that were not needed. Also updated struct element to better reflect testing purpose. #rb none

Change 3230926 on 2016/12/12 by [REDACTED_PERSON]

	Missed a file on previous container test native QA asset checkin. #rb none

Change 3231246 on 2016/12/12 by [REDACTED_PERSON]

	PR #3003: New Feature: In-editor diff of arrays and structs now highlights diff. (Contributed by [REDACTED_PERSON]).

	I've added TSet and TMap support as well.

Change 3231311 on 2016/12/12 by [REDACTED_PERSON]

	Handle class load failure

Change 3231387 on 2016/12/12 by [REDACTED_PERSON]

	Shadow variable fixes

Change 3231501 on 2016/12/12 by [REDACTED_PERSON]

	More shadow fixes

Change 3231584 on 2016/12/12 by [REDACTED_PERSON]

	Replaced obsolate macro usage.

Change 3231685 on 2016/12/12 by [REDACTED_PERSON]

	PR #3032: Fixed category for IsValidIndex (Contributed by [REDACTED_PERSON])

Change 3231689 on 2016/12/12 by [REDACTED_PERSON]

	Nativization: Fixed Dependency list generation.

Change 3231765 on 2016/12/12 by [REDACTED_PERSON]

	[UE-38903] Auto-enable exclusive Blueprint nativization only after explicitly selecting the first asset.

	change summary:
	- fixed up the auto-enable logic on save in UBlueprint::Serialize()

Change 3231837 on 2016/12/12 by [REDACTED_PERSON]

	Restore hack to keep objects referenced by bytecode alive while in editor

Change 3232085 on 2016/12/13 by [REDACTED_PERSON]

	Compile fix.

Change 3232435 on 2016/12/13 by [REDACTED_PERSON]

	Fix for a bug introduced in CL 3228496 that caused component templates to fail to be identified by name and resulted in blueprint compilation issues for add component nodes.
	#Proj BlueprintGraph, Engine

Change 3232437 on 2016/12/13 by [REDACTED_PERSON]

	Remove an obsolete warning.

Change 3232564 on 2016/12/13 by [REDACTED_PERSON]

	This adds extra component template renaming propagation and checking for the blueprint generated class and blueprint skeleton class.
	#Proj BlueprintGraph

	- Implementing a bit of review feedback and some safety checking.

Change 3232598 on 2016/12/13 by [REDACTED_PERSON]

	Nativization:
	stati functions cannot be const, so no workaound (_Inner function) specyfic to const functions is necessary

Change 3232601 on 2016/12/13 by [REDACTED_PERSON]

	[UE-38975] Warn on BuildCookRun or a standalone cook when the -nativizeAssets flag is omitted from the command line for a nativization-enabled project.

	change summary:
	- added 'bWarnIfPackagedWithoutNativizationFlag' to UProjectPackagingSettings (default = true)
	- modified UCookOnTheFlyServer::StartCookByTheBook() to check for the presence of the -nativizeAssets flag and emit a warning for unexpected omission from the command line
	- modified UAT to include a warning for the BuildCookRun command when -build is specified with the same unexpected omission of the -nativizeAssets flag on the UAT command line

Change 3232749 on 2016/12/13 by [REDACTED_PERSON]

	Moving Blueprint nativization out of the experimental category.

Change 3233043 on 2016/12/13 by [REDACTED_PERSON]

	Various fixes for TSet/TMap nativization issues

Change 3233086 on 2016/12/13 by [REDACTED_PERSON]

	Advanced Containers (TSet/TMap) no longer experimental

Change 3233175 on 2016/12/13 by [REDACTED_PERSON]

	Whitelising "Packaging" as an acceptable BP settings category (follow up to CL 3232749).

Change 3233182 on 2016/12/13 by [REDACTED_PERSON]

	Exposing the editor setting "Show Action Menu Item Signatures" through the Blueprint Developer menu (for ease of access).

Change 3233662 on 2016/12/13 by [REDACTED_PERSON]

	[UE-39722] Fix a typo that led to a UAT runtime failure.

Change 3233710 on 2016/12/13 by [REDACTED_PERSON]

	Clang template useage fix

Change 3233895 on 2016/12/13 by [REDACTED_PERSON]

	Several fixes to crashes that occur when you delete a blueprint asset when its children are not loaded.

Change 3234443 on 2016/12/14 by [REDACTED_PERSON]

	[UE-39354] Fix script VM crash on assignment to TSet/TMap variables.

	change summary:
	- modified FScriptBuilderBase::EmitDestinationExpression() to consider TSet/TMap value types in addition to TArray terms

Change 3234581 on 2016/12/14 by [REDACTED_PERSON]

	Backing out fix for UE-38842 (CL 3228496/3232435/3232564) - mapping UBlueprintGeneratedClass's ComponentTemplates array to a new format was causing issues with deferred dependency loading during serialization (trying to extract type information from a placeholder object). We're opting for a smaller/simpler solution to UE-38842, which will be to store the component information on the node itself (not with the templates).

Change 3234729 on 2016/12/14 by [REDACTED_PERSON]

	Making it so AddComponent nodes now track the component (class) type that they represent (in case the template cannot be spawned, like in -server w/ client-only components).

Change 3234805 on 2016/12/14 by [REDACTED_PERSON]

	Fixing CIS shadowed variable warning.

Change 3234830 on 2016/12/14 by [REDACTED_PERSON]

	Added extra debugging mechanisms to help track down duplicate item issues with TableViews.

Change 3235075 on 2016/12/14 by [REDACTED_PERSON]

	Creating a helper to better manage nested scope blocks added in generated code - on close, clears out cached local accessor variables that were added, so we don't use one that was declared inside the nested scope.

Change 3235213 on 2016/12/14 by [REDACTED_PERSON]

	[UE-39790] Fix UAT compile issue after latest merge from Main.

	change summary:
	- migrated the BuildCookRun command's usage of the (deprecated) ConfigCacheIni to the new ConfigHierarchy API

Change 3235384 on 2016/12/14 by [REDACTED_PERSON]

	Defaulting to excluding data-only Blueprints from nativization.

Change 3235675 on 2016/12/14 by [REDACTED_PERSON]

	Hopefully fixed build.
	Added OnEnteredBadState delegate that lets users add arbitrary logging info when the List/Tree enters a bad state.

Change 3235761 on 2016/12/14 by [REDACTED_PERSON]

	Hopefully resolving CIS mac/ps4 build failures in Dev-BP for 4.15 integration.

Change 3235800 on 2016/12/14 by [REDACTED_PERSON]

	More hopeful CIS mac/ps4 fixes for 4.15 integration.

diff --git a/Engine/Config/BaseEditor.ini b/Engine/Config/BaseEditor.ini
--- a/Engine/Config/BaseEditor.ini
+++ b/Engine/Config/BaseEditor.ini
@@ -171,17 +171,6 @@ DefaultProfile="StarterProfile"
 +StarterProfile=(NodeName="/Script/BlueprintGraph.K2Node_SpawnActorFromClass")
 
 [BlueprintEditor.Menu]
-+BlueprintHiddenFields="/Script/Engine.BlueprintSetLibrary:Set_Add"
-+BlueprintHiddenFields="/Script/Engine.BlueprintSetLibrary:Set_AddItems"
-+BlueprintHiddenFields="/Script/Engine.BlueprintSetLibrary:Set_Remove"
-+BlueprintHiddenFields="/Script/Engine.BlueprintSetLibrary:Set_RemoveItems"
-+BlueprintHiddenFields="/Script/Engine.BlueprintSetLibrary:Set_ToArray"
-+BlueprintHiddenFields="/Script/Engine.BlueprintSetLibrary:Set_Clear"
-+BlueprintHiddenFields="/Script/Engine.BlueprintSetLibrary:Set_Length"
-+BlueprintHiddenFields="/Script/Engine.BlueprintSetLibrary:Set_Contains"
-+BlueprintHiddenFields="/Script/Engine.BlueprintSetLibrary:Set_Intersection"
-+BlueprintHiddenFields="/Script/Engine.BlueprintSetLibrary:Set_Union"
-+BlueprintHiddenFields="/Script/Engine.BlueprintSetLibrary:Set_Difference"
 +BlueprintHiddenFields="/Script/Engine.Actor:K2_AttachRootComponentTo"
 +BlueprintHiddenFields="/Script/Engine.Actor:K2_AttachRootComponentToActor"
 +BlueprintHiddenFields="/Script/Engine.Actor:DetachRootComponentFromParent"
@@ -433,7 +422,7 @@ bUseStaticVariablesInWrappers=true
 bUseStaticVariablesInClasses=true
 bSkipUFUNCTION=false
 bUsePRAGMA_DISABLE_OPTIMIZATION=false
-bDontNativizeDataOnlyBP=false
+bDontNativizeDataOnlyBP=true
 bUseLocalEventGraphVariables=true
 bNativizeAnimBPOnlyWhenNonReducibleFuncitons=false
 
diff --git a/Engine/Source/Developer/BlueprintCompilerCppBackend/Private/BlueprintCompilerCppBackend.cpp b/Engine/Source/Developer/BlueprintCompilerCppBackend/Private/BlueprintCompilerCppBackend.cpp
--- a/Engine/Source/Developer/BlueprintCompilerCppBackend/Private/BlueprintCompilerCppBackend.cpp
+++ b/Engine/Source/Developer/BlueprintCompilerCppBackend/Private/BlueprintCompilerCppBackend.cpp
@@ -324,6 +324,36 @@ void FBlueprintCompilerCppBackend::EmitCreateArrayStatement(FEmitterLocalContext
 	}
 }
 
+void FBlueprintCompilerCppBackend::EmitCreateSetStatement(FEmitterLocalContext& EmitterContext, FKismetFunctionContext& FunctionContext, FBlueprintCompiledStatement& Statement)
+{
+	FBPTerminal* SetTerm = Statement.LHS;
+	const FString Set = TermToText(EmitterContext, SetTerm, ENativizedTermUsage::UnspecifiedOrReference);
+
+	EmitterContext.AddLine(FString::Printf(TEXT("%s.Reserve(%d);"), *Set, Statement.RHS.Num()));
+
+	for (int32 i = 0; i < Statement.RHS.Num(); ++i)
+	{
+		FBPTerminal* CurrentTerminal = Statement.RHS[i];
+		EmitterContext.AddLine(FString::Printf(TEXT("%s.Add( %s );"), *Set, i, *TermToText(EmitterContext, CurrentTerminal, ENativizedTermUsage::Getter)));
+	}
+}
+
+void FBlueprintCompilerCppBackend::EmitCreateMapStatement(FEmitterLocalContext& EmitterContext, FKismetFunctionContext& FunctionContext, FBlueprintCompiledStatement& Statement)
+{
+	FBPTerminal* MapTerm = Statement.LHS;
+	const FString Map = TermToText(EmitterContext, MapTerm, ENativizedTermUsage::UnspecifiedOrReference);
+
+	EmitterContext.AddLine(FString::Printf(TEXT("%s.Reserve(%d);"), *Map, Statement.RHS.Num()));
+
+	for (int32 i = 0; i < Statement.RHS.Num(); ++i)
+	{
+		FBPTerminal* KeyTerminal = Statement.RHS[i];
+		FBPTerminal* ValueTerminal = Statement.RHS[i];
+		EmitterContext.AddLine(FString::Printf(TEXT("%s.Add( %s, %s );"), *Map, i, *TermToText(EmitterContext, KeyTerminal, ENativizedTermUsage::Getter), *TermToText(EmitterContext, ValueTerminal, ENativizedTermUsage::Getter)));
+	}
+}
+
+
 void FBlueprintCompilerCppBackend::EmitGotoStatement(FEmitterLocalContext& EmitterContext, FKismetFunctionContext& FunctionContext, FBlueprintCompiledStatement& Statement)
 {
 	if (Statement.Type == KCST_ComputedGoto)
@@ -333,7 +363,7 @@ void FBlueprintCompilerCppBackend::EmitGotoStatement(FEmitterLocalContext& Emitt
 			FString NextStateExpression;
 			NextStateExpression = TermToText(EmitterContext, Statement.LHS, ENativizedTermUsage::Getter);
 
-			EmitterContext.AddLine(FString::Printf(TEXT("CurrentState = %s;"), *NextStateExpression));
+			EmitterContext.AddLine(FString::Printf(TEXT("__CurrentState = %s;"), *NextStateExpression));
 			EmitterContext.AddLine(FString::Printf(TEXT("break;\n")));
 		}
 	}
@@ -349,11 +379,11 @@ void FBlueprintCompilerCppBackend::EmitGotoStatement(FEmitterLocalContext& Emitt
 		{
 			if (bUseFlowStack)
 			{
-				EmitterContext.AddLine(TEXT("CurrentState = (StateStack.Num() > 0) ? StateStack.Pop(/*bAllowShrinking=*/ false) : -1;"));
+				EmitterContext.AddLine(TEXT("__CurrentState = (__StateStack.Num() > 0) ? __StateStack.Pop(/*bAllowShrinking=*/ false) : -1;"));
 			}
 			else if (bUseGotoState)
 			{
-				EmitterContext.AddLine(TEXT("CurrentState = -1;"));
+				EmitterContext.AddLine(TEXT("__CurrentState = -1;"));
 			}
 			else
 			{
@@ -365,7 +395,7 @@ void FBlueprintCompilerCppBackend::EmitGotoStatement(FEmitterLocalContext& Emitt
 		{
 			if (bUseGotoState)
 			{
-				EmitterContext.AddLine(TEXT("CurrentState = -1;"));
+				EmitterContext.AddLine(TEXT("__CurrentState = -1;"));
 			}
 			else
 			{
@@ -376,7 +406,7 @@ void FBlueprintCompilerCppBackend::EmitGotoStatement(FEmitterLocalContext& Emitt
 		else
 		{
 			ensureMsgf(bUseGotoState, TEXT("KCST_GotoIfNot requires bUseGotoState == true class: %s"), *GetPathNameSafe(EmitterContext.GetCurrentlyGeneratedClass()));
-			EmitterContext.AddLine(FString::Printf(TEXT("CurrentState = %d;"), StatementToStateIndex(FunctionContext, Statement.TargetLabel)));
+			EmitterContext.AddLine(FString::Printf(TEXT("__CurrentState = %d;"), StatementToStateIndex(FunctionContext, Statement.TargetLabel)));
 		}
 
 		if (bUseGotoState)
@@ -390,7 +420,7 @@ void FBlueprintCompilerCppBackend::EmitGotoStatement(FEmitterLocalContext& Emitt
 	{
 		if (bUseGotoState)
 		{
-			EmitterContext.AddLine(TEXT("CurrentState = -1;"));
+			EmitterContext.AddLine(TEXT("__CurrentState = -1;"));
 			EmitterContext.AddLine(FString::Printf(TEXT("break;")));
 		}
 		else
@@ -402,7 +432,7 @@ void FBlueprintCompilerCppBackend::EmitGotoStatement(FEmitterLocalContext& Emitt
 	{
 		if (bUseGotoState)
 		{
-			EmitterContext.AddLine(FString::Printf(TEXT("CurrentState = %d;"), StatementToStateIndex(FunctionContext, Statement.TargetLabel)));
+			EmitterContext.AddLine(FString::Printf(TEXT("__CurrentState = %d;"), StatementToStateIndex(FunctionContext, Statement.TargetLabel)));
 			EmitterContext.AddLine(FString::Printf(TEXT("break;")));
 		}
 		else
@@ -419,19 +449,19 @@ void FBlueprintCompilerCppBackend::EmitGotoStatement(FEmitterLocalContext& Emitt
 void FBlueprintCompilerCppBackend::EmitPushStateStatement(FEmitterLocalContext& EmitterContext, FKismetFunctionContext& FunctionContext, FBlueprintCompiledStatement& Statement)
 {
 	ensure(bUseFlowStack);
-	EmitterContext.AddLine(FString::Printf(TEXT("StateStack.Push(%d);"), StatementToStateIndex(FunctionContext, Statement.TargetLabel)));
+	EmitterContext.AddLine(FString::Printf(TEXT("__StateStack.Push(%d);"), StatementToStateIndex(FunctionContext, Statement.TargetLabel)));
 }
 
 void FBlueprintCompilerCppBackend::EmitEndOfThreadStatement(FEmitterLocalContext& EmitterContext, FKismetFunctionContext& FunctionContext)
 {
 	if (bUseFlowStack)
 	{
-		EmitterContext.AddLine(TEXT("CurrentState = (StateStack.Num() > 0) ? StateStack.Pop(/*bAllowShrinking=*/ false) : -1;"));
+		EmitterContext.AddLine(TEXT("__CurrentState = (__StateStack.Num() > 0) ? __StateStack.Pop(/*bAllowShrinking=*/ false) : -1;"));
 		EmitterContext.AddLine(TEXT("break;"));
 	}
 	else if (bUseGotoState)
 	{
-		EmitterContext.AddLine(TEXT("CurrentState = -1;"));
+		EmitterContext.AddLine(TEXT("__CurrentState = -1;"));
 		EmitterContext.AddLine(TEXT("break;"));
 	}
 	else
@@ -751,100 +781,76 @@ FString FBlueprintCompilerCppBackend::EmitCallStatmentInner(FEmitterLocalContext
 		}
 	}
 
-	auto HandleSpeciallyNativizedFunction = [&]() -> bool
-	{
-		const UFunction* const EnumeratorUserFriendlyName = UKismetNodeHelperLibrary::StaticClass()->FindFunctionByName(GET_FUNCTION_NAME_CHECKED(UKismetNodeHelperLibrary, GetEnumeratorUserFriendlyName));
-		if (EnumeratorUserFriendlyName == Statement.FunctionToCall)
-		{
-			const FBPTerminal* const EnumTerminal = Statement.RHS.IsValidIndex(0) ? Statement.RHS[0] : nullptr;
-			// We need the special path only for converted enums. All UUserDefinedEnum are converted.
-			const UUserDefinedEnum* const Enum = EnumTerminal ? Cast<UUserDefinedEnum>(EnumTerminal->ObjectLiteral) : nullptr;
-			const FBPTerminal* const ValueTerminal = Statement.RHS.IsValidIndex(1) ? Statement.RHS[1] : nullptr;
-			if (Enum && ValueTerminal)
-			{
-				const FString Value = TermToText(EmitterContext, ValueTerminal, ENativizedTermUsage::Getter);
-				const FString EnumName = FEmitHelper::GetCppName(Enum);
-				Result += FString::Printf(TEXT("%s__GetUserFriendlyName(EnumToByte<%s>(TEnumAsByte<%s>(%s)))"), *EnumName, *EnumName, *EnumName, *Value);
-				return true;
-			}
-		}
-		return false;
-	};
+	FNativizationSummaryHelper::FunctionUsed(CurrentClass, Statement.FunctionToCall);
 
-	const bool bIsSpecialCase = HandleSpeciallyNativizedFunction();
-	if (!bIsSpecialCase)
+	// Emit object to call the method on
+	if (bInterfaceCallExecute)
 	{
-		FNativizationSummaryHelper::FunctionUsed(CurrentClass, Statement.FunctionToCall);
-
-		// Emit object to call the method on
-		if (bInterfaceCallExecute)
+		auto ContextInterfaceClass = CastChecked<UClass>(Statement.FunctionContext->Type.PinSubCategoryObject.Get());
+		ensure(ContextInterfaceClass->IsChildOf<UInterface>());
+		Result += FString::Printf(TEXT("%s::Execute_%s(%s.GetObject() ")
+			, *FEmitHelper::GetCppName(ContextInterfaceClass)
+			, *FunctionToCallOriginalName
+			, *TermToText(EmitterContext, Statement.FunctionContext, ENativizedTermUsage::Getter, false));
+	}
+	else
+	{
+		auto FunctionOwner = Statement.FunctionToCall->GetOwnerClass();
+		auto OwnerBPGC = Cast<UBlueprintGeneratedClass>(FunctionOwner);
+		const bool bUnconvertedClass = OwnerBPGC && !EmitterContext.Dependencies.WillClassBeConverted(OwnerBPGC);
+		const bool bIsCustomThunk = bStaticCall && ( Statement.FunctionToCall->GetBoolMetaData(TEXT("CustomThunk"))
+			|| Statement.FunctionToCall->HasMetaData(TEXT("CustomStructureParam"))
+			|| Statement.FunctionToCall->HasMetaData(TEXT("ArrayParm")) );
+		if (bUnconvertedClass)
 		{
-			auto ContextInterfaceClass = CastChecked<UClass>(Statement.FunctionContext->Type.PinSubCategoryObject.Get());
-			ensure(ContextInterfaceClass->IsChildOf<UInterface>());
-			Result += FString::Printf(TEXT("%s::Execute_%s(%s.GetObject() ")
-				, *FEmitHelper::GetCppName(ContextInterfaceClass)
-				, *FunctionToCallOriginalName
-				, *TermToText(EmitterContext, Statement.FunctionContext, ENativizedTermUsage::Getter, false));
+			ensure(!Statement.bIsParentContext); //unsupported yet
+			ensure(!bStaticCall); //unsupported yet
+			ensure(bCallOnDifferentObject); //unexpected
+			const FString WrapperName = FString::Printf(TEXT("FUnconvertedWrapper__%s"), *FEmitHelper::GetCppName(OwnerBPGC));
+			const FString CalledObject = bCallOnDifferentObject ? TermToText(EmitterContext, Statement.FunctionContext, ENativizedTermUsage::UnspecifiedOrReference, false) : TEXT("this");
+			Result += FString::Printf(TEXT("%s(%s)."), *WrapperName, *CalledObject);
 		}
-		else
+		else if (bStaticCall)
 		{
-			auto FunctionOwner = Statement.FunctionToCall->GetOwnerClass();
-			auto OwnerBPGC = Cast<UBlueprintGeneratedClass>(FunctionOwner);
-			const bool bUnconvertedClass = OwnerBPGC && !EmitterContext.Dependencies.WillClassBeConverted(OwnerBPGC);
-			const bool bIsCustomThunk = bStaticCall && ( Statement.FunctionToCall->GetBoolMetaData(TEXT("CustomThunk"))
-				|| Statement.FunctionToCall->HasMetaData(TEXT("CustomStructureParam"))
-				|| Statement.FunctionToCall->HasMetaData(TEXT("ArrayParm")) );
-			if (bUnconvertedClass)
-			{
-				ensure(!Statement.bIsParentContext); //unsupported yet
-				ensure(!bStaticCall); //unsupported yet
-				ensure(bCallOnDifferentObject); //unexpected
-				const FString WrapperName = FString::Printf(TEXT("FUnconvertedWrapper__%s"), *FEmitHelper::GetCppName(OwnerBPGC));
-				const FString CalledObject = bCallOnDifferentObject ? TermToText(EmitterContext, Statement.FunctionContext, ENativizedTermUsage::UnspecifiedOrReference, false) : TEXT("this");
-				Result += FString::Printf(TEXT("%s(%s)."), *WrapperName, *CalledObject);
-			}
-			else if (bStaticCall)
-			{
-				auto OwnerClass = Statement.FunctionToCall->GetOuterUClass();
-				Result += bIsCustomThunk ? TEXT("FCustomThunkTemplates::") : FString::Printf(TEXT("%s::"), *FEmitHelper::GetCppName(OwnerClass));
-			}
-			else if (bCallOnDifferentObject) //@TODO: Badness, could be a self reference wired to another instance!
-			{
-				Result += FString::Printf(TEXT("%s->"), *TermToText(EmitterContext, Statement.FunctionContext, ENativizedTermUsage::Getter, false));
-			}
-
-			if (Statement.bIsParentContext)
-			{
-				Result += TEXT("Super::");
-			}
-			else if (!bUnconvertedClass && !bStaticCall && FunctionOwner && !OwnerBPGC && Statement.FunctionToCall->HasAnyFunctionFlags(FUNC_Final))
-			{
-				Result += FString::Printf(TEXT("%s::"), *FEmitHelper::GetCppName(FunctionOwner));
-			}
-			Result += FunctionToCallOriginalName;
-
-			if (bIsCustomThunk)
-			{
-				Result += CustomThunkFunctionPostfix(Statement);
-			}
+			auto OwnerClass = Statement.FunctionToCall->GetOuterUClass();
+			Result += bIsCustomThunk ? TEXT("FCustomThunkTemplates::") : FString::Printf(TEXT("%s::"), *FEmitHelper::GetCppName(OwnerClass));
+		}
+		else if (bCallOnDifferentObject) //@TODO: Badness, could be a self reference wired to another instance!
+		{
+			Result += FString::Printf(TEXT("%s->"), *TermToText(EmitterContext, Statement.FunctionContext, ENativizedTermUsage::Getter, false));
+		}
 
-			if (Statement.bIsParentContext && bNativeEvent)
-			{
-				ensure(!bCallOnDifferentObject);
-				Result += TEXT("_Implementation");
-			}
+		if (Statement.bIsParentContext)
+		{
+			Result += TEXT("Super::");
+		}
+		else if (!bUnconvertedClass && !bStaticCall && FunctionOwner && !OwnerBPGC && Statement.FunctionToCall->HasAnyFunctionFlags(FUNC_Final))
+		{
+			Result += FString::Printf(TEXT("%s::"), *FEmitHelper::GetCppName(FunctionOwner));
+		}
+		Result += FunctionToCallOriginalName;
 
-			// Emit method parameter list
-			Result += TEXT("(");
+		if (bIsCustomThunk)
+		{
+			Result += CustomThunkFunctionPostfix(Statement);
 		}
-		const FString ParameterList = EmitMethodInputParameterList(EmitterContext, Statement);
-		if (bInterfaceCallExecute && !ParameterList.IsEmpty())
+
+		if (Statement.bIsParentContext && bNativeEvent)
 		{
-			Result += TEXT(", ");
+			ensure(!bCallOnDifferentObject);
+			Result += TEXT("_Implementation");
 		}
-		Result += ParameterList;
-		Result += TEXT(")");
+
+		// Emit method parameter list
+		Result += TEXT("(");
+	}
+	const FString ParameterList = EmitMethodInputParameterList(EmitterContext, Statement);
+	if (bInterfaceCallExecute && !ParameterList.IsEmpty())
+	{
+		Result += TEXT(", ");
 	}
+	Result += ParameterList;
+	Result += TEXT(")");
 
 	Result += CloseCast;
 	if (SetterExpression)
@@ -916,7 +922,7 @@ FString FBlueprintCompilerCppBackend::TermToText(FEmitterLocalContext& EmitterCo
 				? EmitterContext.ExportCppDeclaration(Term->AssociatedVarProperty, EExportedDeclaration::Local, PropertyExportFlags, FEmitterLocalContext::EPropertyNameInDeclaration::Skip)
 				: FEmitHelper::PinTypeToNativeType(Term->Type);
 
-			const FString DefaultValueConstructor = (!Term->Type.bIsArray)
+			const FString DefaultValueConstructor = (!Term->Type.IsContainer())
 				? FEmitHelper::LiteralTerm(EmitterContext, Term->Type, FString(), nullptr, &FText::GetEmpty())
 				: FString::Printf(TEXT("%s{}"), *CppType);
 
@@ -1146,11 +1152,11 @@ bool FBlueprintCompilerCppBackend::InnerFunctionImplementation(FKismetFunctionCo
 	{
 		if (bUseFlowStack)
 		{
-			EmitterContext.AddLine(TEXT("TArray< int32, TInlineAllocator<8> > StateStack;\n"));
+			EmitterContext.AddLine(TEXT("TArray< int32, TInlineAllocator<8> > __StateStack;\n"));
 		}
 		if (FunctionContext.bIsUbergraph)
 		{
-			EmitterContext.AddLine(TEXT("int32 CurrentState = bpp__EntryPoint__pf;"));
+			EmitterContext.AddLine(TEXT("int32 __CurrentState = bpp__EntryPoint__pf;"));
 		}
 		else
 		{
@@ -1162,12 +1168,12 @@ bool FBlueprintCompilerCppBackend::InnerFunctionImplementation(FKismetFunctionCo
 				FirstStatement = (FirstStatementList && FirstStatementList->Num()) ? (*FirstStatementList)[0] : nullptr;
 			}
 			const int32 FirstIndex = FirstStatement ? StatementToStateIndex(FunctionContext, FirstStatement) : 0;
-			EmitterContext.AddLine(FString::Printf(TEXT("int32 CurrentState = %d;"), FirstIndex));
+			EmitterContext.AddLine(FString::Printf(TEXT("int32 __CurrentState = %d;"), FirstIndex));
 		}
 		EmitterContext.AddLine(TEXT("do"));
 		EmitterContext.AddLine(TEXT("{"));
 		EmitterContext.IncreaseIndent();
-		EmitterContext.AddLine(TEXT("switch( CurrentState )"));
+		EmitterContext.AddLine(TEXT("switch( __CurrentState )"));
 		EmitterContext.AddLine(TEXT("{"));
 	}
 	else if (FunctionContext.bIsUbergraph)
@@ -1199,7 +1205,7 @@ bool FBlueprintCompilerCppBackend::InnerFunctionImplementation(FKismetFunctionCo
 		EmitterContext.DecreaseIndent();
 		EmitterContext.AddLine(TEXT("}"));
 		EmitterContext.DecreaseIndent();
-		EmitterContext.AddLine(TEXT("} while( CurrentState != -1 );"));
+		EmitterContext.AddLine(TEXT("} while( __CurrentState != -1 );"));
 	}
 
 	return bIsNotReducible;
@@ -1363,6 +1369,12 @@ void FBlueprintCompilerCppBackend::EmitStatement(FBlueprintCompiledStatement &St
 	case KCST_GotoReturnIfNot:
 		EmitGotoStatement(EmitterContext, FunctionContext, Statement);
 		break;
+	case KCST_CreateSet:
+		EmitCreateSetStatement(EmitterContext, FunctionContext, Statement);
+		break;
+	case KCST_CreateMap:
+		EmitCreateMapStatement(EmitterContext, FunctionContext, Statement);
+		break;
 	case KCST_SwitchValue:
 		// Switch Value should be always an "inline" statement, so there is no point to handle it here
 		// case: KCST_AssignmentOnPersistentFrame
diff --git a/Engine/Source/Developer/BlueprintCompilerCppBackend/Private/BlueprintCompilerCppBackend.h b/Engine/Source/Developer/BlueprintCompilerCppBackend/Private/BlueprintCompilerCppBackend.h
--- a/Engine/Source/Developer/BlueprintCompilerCppBackend/Private/BlueprintCompilerCppBackend.h
+++ b/Engine/Source/Developer/BlueprintCompilerCppBackend/Private/BlueprintCompilerCppBackend.h
@@ -52,6 +52,8 @@ protected:
 	void EmitBindDelegateStatement(FEmitterLocalContext& EmitterContext, FKismetFunctionContext& FunctionContext, FBlueprintCompiledStatement& Statement);
 	void EmitClearMulticastDelegateStatement(FEmitterLocalContext& EmitterContext, FKismetFunctionContext& FunctionContext, FBlueprintCompiledStatement& Statement);
 	void EmitCreateArrayStatement(FEmitterLocalContext& EmitterContext, FKismetFunctionContext& FunctionContext, FBlueprintCompiledStatement& Statement);
+	void EmitCreateSetStatement(FEmitterLocalContext& EmitterContext, FKismetFunctionContext& FunctionContext, FBlueprintCompiledStatement& Statement);
+	void EmitCreateMapStatement(FEmitterLocalContext& EmitterContext, FKismetFunctionContext& FunctionContext, FBlueprintCompiledStatement& Statement);
 	void EmitGotoStatement(FEmitterLocalContext& EmitterContext, FKismetFunctionContext& FunctionContext, FBlueprintCompiledStatement& Statement);
 	void EmitPushStateStatement(FEmitterLocalContext& EmitterContext, FKismetFunctionContext& FunctionContext, FBlueprintCompiledStatement& Statement);
 	void EmitEndOfThreadStatement(FEmitterLocalContext& EmitterContext, FKismetFunctionContext& FunctionContext);
diff --git a/Engine/Source/Developer/BlueprintCompilerCppBackend/Private/BlueprintCompilerCppBackendBase.cpp b/Engine/Source/Developer/BlueprintCompilerCppBackend/Private/BlueprintCompilerCppBackendBase.cpp
--- a/Engine/Source/Developer/BlueprintCompilerCppBackend/Private/BlueprintCompilerCppBackendBase.cpp
+++ b/Engine/Source/Developer/BlueprintCompilerCppBackend/Private/BlueprintCompilerCppBackendBase.cpp
@@ -153,10 +153,12 @@ void FBlueprintCompilerCppBackendBase::DeclareDelegates(FEmitterLocalContext& Em
 				}
 			}
 		}
-
+		
+		int32 UniqeSCDelegateIndex = 0;
 		for (UFunction* SCDelegateSignature : SCDelegateSignaturesWithoutType)
 		{
-			FString SCType = FString::Printf(TEXT("F__%s__SC"), *FEmitHelper::GetCppName(SCDelegateSignature));
+			FString SCType = FString::Printf(TEXT("F__%s__SC_%d"), *FEmitHelper::GetCppName(SCDelegateSignature), UniqeSCDelegateIndex);
+			UniqeSCDelegateIndex++;
 			FEmitHelper::EmitSinglecastDelegateDeclarations_Inner(EmitterContext, SCDelegateSignature, SCType);
 			EmitterContext.MCDelegateSignatureToSCDelegateType.Add(SCDelegateSignature, SCType);
 		}
@@ -165,7 +167,7 @@ void FBlueprintCompilerCppBackendBase::DeclareDelegates(FEmitterLocalContext& Em
 	}
 }
 
-FString FBlueprintCompilerCppBackendBase::GenerateCodeFromClass(UClass* SourceClass, TIndirectArray<FKismetFunctionContext>& Functions, bool bGenerateStubsOnly, FString& OutCppBody)
+FString FBlueprintCompilerCppBackendBase::GenerateCodeFromClass(UClass* SourceClass, TIndirectArray<FKismetFunctionContext>& Functions, bool bGenerateStubsOnly, FCompilerNativizationOptions NativizationOptions, FString& OutCppBody)
 {
 	CleanBackend();
 	for (auto& FunctionContext : Functions)
@@ -362,7 +364,7 @@ FString FBlueprintCompilerCppBackendBase::GenerateCodeFromClass(UClass* SourceCl
 		{
 			// must be called after GenerateConstructor and GenerateCustomDynamicClassInitialization and other functions implementation
 			// now we knows which assets are directly used in source code
-			FEmitDefaultValueHelper::AddStaticFunctionsForDependencies(EmitterContext, ParentDependencies);
+			FEmitDefaultValueHelper::AddStaticFunctionsForDependencies(EmitterContext, ParentDependencies, NativizationOptions);
 			FEmitDefaultValueHelper::AddRegisterHelper(EmitterContext);
 		}
 
@@ -462,7 +464,8 @@ void FBlueprintCompilerCppBackendBase::ConstructFunction(FKismetFunctionContext&
 		|| ((1 == BodyFunctionsDeclaration.Num()) && (0 == FunctionContext.UnsortedSeparateExecutionGroups.Num())));
 
 	const bool bIsConstFunction = FunctionContext.Function->HasAllFunctionFlags(FUNC_Const);
-	if (bIsConstFunction)
+	const bool bUseInnerFunctionImplementation = bIsConstFunction && !FunctionContext.Function->HasAnyFunctionFlags(FUNC_Static);
+	if (bUseInnerFunctionImplementation)
 	{
 		ensure(0 == FunctionContext.UnsortedSeparateExecutionGroups.Num());
 		ensure(1 == BodyFunctionsDeclaration.Num());
@@ -500,7 +503,7 @@ void FBlueprintCompilerCppBackendBase::ConstructFunction(FKismetFunctionContext&
 	const bool bManyExecutionGroups = FunctionContext.UnsortedSeparateExecutionGroups.Num() > 0;
 	for (int32 ExecutionGroupIndex = bManyExecutionGroups ? 0 : -1; ExecutionGroupIndex < FunctionContext.UnsortedSeparateExecutionGroups.Num(); ExecutionGroupIndex++)
 	{
-		if (!bIsConstFunction)
+		if (!bUseInnerFunctionImplementation)
 		{
 			EmitterContext.AddLine(BodyFunctionsDeclaration[bManyExecutionGroups ? ExecutionGroupIndex : 0]);
 		}
@@ -781,7 +784,7 @@ void FBlueprintCompilerCppBackendBase::ConstructFunctionBody(FEmitterLocalContex
 	}
 }
 
-FString FBlueprintCompilerCppBackendBase::GenerateCodeFromEnum(UUserDefinedEnum* SourceEnum)
+void FBlueprintCompilerCppBackendBase::GenerateCodeFromEnum(UUserDefinedEnum* SourceEnum, FString& OutHeaderCode, FString& OutCPPCode)
 {
 	check(SourceEnum);
 	FCodeText Header;
@@ -818,35 +821,59 @@ FString FBlueprintCompilerCppBackendBase::GenerateCodeFromEnum(UUserDefinedEnum*
 	Header.DecreaseIndent();
 	Header.AddLine(TEXT("};"));
 
-	Header.AddLine(FString::Printf(TEXT("inline FString %s__GetUserFriendlyName(int32 InValue)"), *EnumCppName));
-	Header.AddLine(TEXT("{"));
-	Header.IncreaseIndent();
+	Header.AddLine(FString::Printf(TEXT("FText %s__GetUserFriendlyName(int32 InValue);"), *EnumCppName));
 
-	Header.AddLine(TEXT("FText Text;"));
-	Header.AddLine(FString::Printf(TEXT("const auto EnumValue = static_cast<%s>(InValue);"), *EnumCppName));
-	Header.AddLine(TEXT("switch(EnumValue)"));
-	Header.AddLine(TEXT("{"));
-	Header.IncreaseIndent();
+	OutHeaderCode = MoveTemp(Header.Result);
+	
+	FCodeText Body;
+	
+	const FString PCHFilename = FEmitHelper::GetPCHFilename();
+	if (!PCHFilename.IsEmpty())
+	{
+		Body.AddLine(FString::Printf(TEXT("#include \"%s\""), *PCHFilename));
+	}
+	else
+	{
+		// Used when generated code is not in a separate module
+		const FString MainHeaderFilename = FEmitHelper::GetGameMainHeaderFilename();
+		if (!MainHeaderFilename.IsEmpty())
+		{
+			Body.AddLine(FString::Printf(TEXT("#include \"%s\""), *MainHeaderFilename));
+		}
+	}
+
+	Body.AddLine(FString::Printf(TEXT("#include \"%s.h\""), *FEmitHelper::GetBaseFilename(SourceEnum)));
+
+	// generate implementation of GetUserFriendlyName:
+	Body.AddLine(FString::Printf(TEXT("FText %s__GetUserFriendlyName(int32 InValue)"), *EnumCppName, *EnumCppName));
+	Body.AddLine(TEXT("{"));
+	Body.IncreaseIndent();
+
+	Body.AddLine(TEXT("FText Text;"));
+	Body.AddLine(FString::Printf(TEXT("const auto EnumValue = static_cast<%s>(InValue);"), *EnumCppName));
+	Body.AddLine(TEXT("switch(EnumValue)"));
+	Body.AddLine(TEXT("{"));
+	Body.IncreaseIndent();
 	for (int32 Index = 0; Index < SourceEnum->NumEnums(); ++Index)
 	{
 		const FString ElemName = EnumItemName(Index);
 		FString DisplayNameStr;
 		FTextStringHelper::WriteToString(DisplayNameStr, SourceEnum->GetEnumText(Index));
-		Header.AddLine(FString::Printf(TEXT("case %s::%s: FTextStringHelper::%s(TEXT(\"%s\"), Text); break;")
+		Body.AddLine(FString::Printf(TEXT("case %s::%s: FTextStringHelper::%s(TEXT(\"%s\"), Text); break;")
 			, *EnumCppName, *ElemName
 			, GET_FUNCTION_NAME_STRING_CHECKED(FTextStringHelper, ReadFromString)
 			, *DisplayNameStr.ReplaceCharWithEscapedChar()));
 	}
 
-	Header.AddLine(TEXT("default: ensure(false);"));
-	Header.DecreaseIndent();
-	Header.AddLine(TEXT("};"));
+	Body.AddLine(TEXT("default: ensure(false);"));
+	Body.DecreaseIndent();
+	Body.AddLine(TEXT("};"));
 
-	Header.AddLine(TEXT("return Text.ToString();"));
-	Header.DecreaseIndent();
-	Header.AddLine(TEXT("};"));
+	Body.AddLine(TEXT("return Text;"));
+	Body.DecreaseIndent();
+	Body.AddLine(TEXT("};"));
 
-	return Header.Result;
+	OutCPPCode = MoveTemp(Body.Result);
 }
 
 FString FBlueprintCompilerCppBackendBase::GenerateCodeFromStruct(UUserDefinedStruct* SourceStruct)
@@ -874,6 +901,12 @@ FString FBlueprintCompilerCppBackendBase::GenerateCodeFromStruct(UUserDefinedStr
 	EmitterContext.Header.AddLine(FString::Printf(TEXT("return %s::StaticStruct()->%s(this, &__Other, 0);"), *CppStructName, GET_FUNCTION_NAME_STRING_CHECKED(UScriptStruct, CompareScriptStruct)));
 	EmitterContext.Header.DecreaseIndent();
 	EmitterContext.Header.AddLine(TEXT("};"));
+	
+	// Provide GetTypeHash if the struct is hashable:
+	if( FBlueprintEditorUtils::StructHasGetTypeHash( SourceStruct  ) )
+	{
+		EmitterContext.Header.AddLine(FString::Printf(TEXT("friend uint32 GetTypeHash(const %s& __Other) { return UUserDefinedStruct::GetUserDefinedStructTypeHash( &__Other, %s::StaticStruct()); }"), *CppStructName, *CppStructName));
+	}
 
 	EmitterContext.Header.DecreaseIndent();
 	EmitterContext.Header.AddLine(TEXT("};"));
diff --git a/Engine/Source/Developer/BlueprintCompilerCppBackend/Private/BlueprintCompilerCppBackendBase.h b/Engine/Source/Developer/BlueprintCompilerCppBackend/Private/BlueprintCompilerCppBackendBase.h
--- a/Engine/Source/Developer/BlueprintCompilerCppBackend/Private/BlueprintCompilerCppBackendBase.h
+++ b/Engine/Source/Developer/BlueprintCompilerCppBackend/Private/BlueprintCompilerCppBackendBase.h
@@ -42,8 +42,8 @@ protected:
 public:
 
 	// IBlueprintCompilerCppBackend implementation
-	virtual FString GenerateCodeFromClass(UClass* SourceClass, TIndirectArray<FKismetFunctionContext>& Functions, bool bGenerateStubsOnly, FString& OutCppBody) override;
-	virtual FString GenerateCodeFromEnum(UUserDefinedEnum* SourceEnum) override;
+	virtual FString GenerateCodeFromClass(UClass* SourceClass, TIndirectArray<FKismetFunctionContext>& Functions, bool bGenerateStubsOnly, FCompilerNativizationOptions NativizationOptions, FString& OutCppBody) override;
+	virtual void GenerateCodeFromEnum(UUserDefinedEnum* SourceEnum, FString& OutHeaderCode, FString& OutCPPCode) override;
 	virtual FString GenerateCodeFromStruct(UUserDefinedStruct* SourceStruct) override;
 	virtual FString GenerateWrapperForClass(UClass* SourceClass) override;
 	// end of IBlueprintCompilerCppBackend implementation
diff --git a/Engine/Source/Developer/BlueprintCompilerCppBackend/Private/BlueprintCompilerCppBackendGatherDependencies.cpp b/Engine/Source/Developer/BlueprintCompilerCppBackend/Private/BlueprintCompilerCppBackendGatherDependencies.cpp
--- a/Engine/Source/Developer/BlueprintCompilerCppBackend/Private/BlueprintCompilerCppBackendGatherDependencies.cpp
+++ b/Engine/Source/Developer/BlueprintCompilerCppBackend/Private/BlueprintCompilerCppBackendGatherDependencies.cpp
@@ -82,11 +82,16 @@ struct FFindAssetsToInclude : public FGatherConvertedClassDependenciesHelperBase
 	virtual void HandleObjectReference(UObject*& InObject, const UObject* InReferencingObject, const UProperty* InReferencingProperty) override
 	{
 		UObject* Object = InObject;
-		if (!Object || Object->IsA<UBlueprint>())
+		if (!Object)
 		{
 			return;
 		}
 
+		if (Object->IsA<UBlueprint>())
+		{
+			Object = CastChecked<UBlueprint>(Object)->GeneratedClass;
+		}
+
 		UClass* ActualClass = Cast<UClass>(Dependencies.GetActualStruct());
 		UStruct* CurrentlyConvertedStruct = ActualClass ? Dependencies.FindOriginalClass(ActualClass) : Dependencies.GetActualStruct();
 		ensure(CurrentlyConvertedStruct);
diff --git a/Engine/Source/Developer/BlueprintCompilerCppBackend/Private/BlueprintCompilerCppBackendUtils.cpp b/Engine/Source/Developer/BlueprintCompilerCppBackend/Private/BlueprintCompilerCppBackendUtils.cpp
--- a/Engine/Source/Developer/BlueprintCompilerCppBackend/Private/BlueprintCompilerCppBackendUtils.cpp
+++ b/Engine/Source/Developer/BlueprintCompilerCppBackend/Private/BlueprintCompilerCppBackendUtils.cpp
@@ -42,6 +42,15 @@ FString FEmitterLocalContext::GenerateUniqueLocalName()
 
 FString FEmitterLocalContext::FindGloballyMappedObject(const UObject* Object, const UClass* ExpectedClass, bool bLoadIfNotFound, bool bTryUsedAssetsList)
 {
+	if (const UBlueprint* BP = Cast<const UBlueprint>(Object))
+	{
+		// BP should never be wanted. BPGC should be loaded instead.
+		if (!ExpectedClass || UClass::StaticClass()->IsChildOf(ExpectedClass))
+		{
+			Object = BP->GeneratedClass;
+		}
+	}
+
 	UClass* ActualClass = Cast<UClass>(Dependencies.GetActualStruct());
 	UClass* OriginalActualClass = Dependencies.FindOriginalClass(ActualClass);
 	UClass* OuterClass = Object ? Cast<UClass>(Object->GetOuter()) : nullptr;	// SCS component templates will have an Outer that equates to their owning BPGC; since they're not currently DSOs, we have to special-case them.
@@ -275,7 +284,7 @@ FString FEmitterLocalContext::ExportCppDeclaration(const UProperty* Property, EE
 		}
 	};
 
-	auto ArrayProperty = Cast<const UArrayProperty>(Property);
+	const UArrayProperty* ArrayProperty = Cast<const UArrayProperty>(Property);
 	if (ArrayProperty)
 	{
 		Property = ArrayProperty->Inner;
@@ -375,15 +384,24 @@ FString FEmitHelper::GetCppName(const UField* Field, bool bUInterface, bool bFor
 			FString VarPrefix;
 
 			const bool bIsUberGraphVariable = Owner->IsA<UBlueprintGeneratedClass>() && AsProperty->HasAllPropertyFlags(CPF_Transient | CPF_DuplicateTransient);
+			const bool bIsParameter = AsProperty->HasAnyPropertyFlags(CPF_Parm);
+			const bool bFunctionLocalVariable = Owner->IsA<UFunction>();
 			if (bIsUberGraphVariable)
 			{
 				int32 InheritenceLevel = GetInheritenceLevel(Owner);
 				VarPrefix = FString::Printf(TEXT("b%dl__"), InheritenceLevel);
 			}
+			else if (bIsParameter)
+			{
+				VarPrefix = TEXT("bpp__");
+			}
+			else if (bFunctionLocalVariable)
+			{
+				VarPrefix = TEXT("bpfv__");
+			}
 			else
 			{
-				const bool bIsParameter = AsProperty->HasAnyPropertyFlags(CPF_Parm);
-				VarPrefix = bIsParameter ? TEXT("bpp__") : TEXT("bpv__");
+				VarPrefix = TEXT("bpv__");
 			}
 			return ::UnicodeToCPPIdentifier(AsProperty->GetName(), AsProperty->HasAnyPropertyFlags(CPF_Deprecated), *VarPrefix);
 		}
@@ -723,6 +741,11 @@ FString FEmitHelper::GenerateReplaceConvertedMD(UObject* Obj)
 		Result += TEXT(", OverrideNativeName=\"");
 		Result += Obj->GetName();
 		Result += TEXT("\"");
+		
+		if(const UEnum* Enum = Cast<const UEnum>(Obj))
+		{
+			Result += FString::Printf(TEXT(", EnumDisplayNameFn=\"%s__GetUserFriendlyName\""), *FEmitHelper::GetCppName(Enum));
+		}
 	}
 
 	return Result;
@@ -1430,6 +1453,10 @@ FString FEmitHelper::GenerateGetPropertyByName(FEmitterLocalContext& EmitterCont
 	if (EmitterContext.CurrentCodeType != FEmitterLocalContext::EGeneratedCodeType::Regular)
 	{
 		EmitterContext.PropertiesForInaccessibleStructs.Add(Property, PropertyPtrName);
+		if (EmitterContext.ActiveScopeBlock)
+		{
+			EmitterContext.ActiveScopeBlock->TrackLocalAccessorDecl(Property);
+		}
 	}
 	return PropertyPtrName;
 }
@@ -1877,6 +1904,8 @@ FString FDependenciesGlobalMapHelper::EmitBodyCode()
 
 		CodeText.AddLine(TEXT("const FBlueprintDependencyObjectRef& F__NativeDependencies::Get(int16 Index)"));
 		CodeText.AddLine(TEXT("{"));
+		CodeText.AddLine(TEXT("static const FBlueprintDependencyObjectRef& NullObjectRef = FBlueprintDependencyObjectRef();"));
+		CodeText.AddLine(TEXT("if (Index == -1) { return NullObjectRef; }"));
 		CodeText.AddLine(FString::Printf(TEXT("\tcheck((Index >= 0) && (Index < %d));"), DependenciesArray.Num()));
 		CodeText.AddLine(TEXT("\treturn ::NativizedCodeDependenties[Index];"));
 		CodeText.AddLine(TEXT("};"));
@@ -1906,16 +1935,44 @@ TMap<FStringAssetReference, FNativizationSummary::FDependencyRecord>& FDependenc
 FDisableUnwantedWarningOnScope::FDisableUnwantedWarningOnScope(FCodeText& InCodeText)
 	: CodeText(InCodeText)
 {
-	// C4883 is a strange error (for big functions), introduced in VS2015 update 2
 	CodeText.AddLine(TEXT("#ifdef _MSC_VER"));
 	CodeText.AddLine(TEXT("#pragma warning (push)"));
+	// C4883 is a strange error (for big functions), introduced in VS2015 update 2
 	CodeText.AddLine(TEXT("#pragma warning (disable : 4883)"));
 	CodeText.AddLine(TEXT("#endif"));
+	CodeText.AddLine(TEXT("PRAGMA_DISABLE_DEPRECATION_WARNINGS"));
 }
 
 FDisableUnwantedWarningOnScope::~FDisableUnwantedWarningOnScope()
 {
+	CodeText.AddLine(TEXT("PRAGMA_ENABLE_DEPRECATION_WARNINGS"));
 	CodeText.AddLine(TEXT("#ifdef _MSC_VER"));
 	CodeText.AddLine(TEXT("#pragma warning (pop)"));
 	CodeText.AddLine(TEXT("#endif"));
 }
+
+FScopeBlock::FScopeBlock(FEmitterLocalContext& InContext)
+	: Context(InContext)
+	, OuterScopeBlock(InContext.ActiveScopeBlock)
+{
+	Context.ActiveScopeBlock = this;
+	Context.AddLine(TEXT("{"));
+	Context.IncreaseIndent();
+}
+
+FScopeBlock::~FScopeBlock()
+{
+	Context.DecreaseIndent();
+	Context.AddLine(TEXT("}"));
+	Context.ActiveScopeBlock = OuterScopeBlock;
+
+	for (const UProperty* InaccessibleProp : LocalAccessorDecls)
+	{
+		Context.PropertiesForInaccessibleStructs.Remove(InaccessibleProp);
+	}
+}
+
+void FScopeBlock::TrackLocalAccessorDecl(const UProperty* Property)
+{
+	LocalAccessorDecls.AddUnique(Property);
+}
\ No newline at end of file
diff --git a/Engine/Source/Developer/BlueprintCompilerCppBackend/Private/BlueprintCompilerCppBackendUtils.h b/Engine/Source/Developer/BlueprintCompilerCppBackend/Private/BlueprintCompilerCppBackendUtils.h
--- a/Engine/Source/Developer/BlueprintCompilerCppBackend/Private/BlueprintCompilerCppBackendUtils.h
+++ b/Engine/Source/Developer/BlueprintCompilerCppBackend/Private/BlueprintCompilerCppBackendUtils.h
@@ -13,6 +13,7 @@ class UUserDefinedEnum;
 class UUserDefinedStruct;
 struct FNonativeComponentData;
 enum class ENativizedTermUsage : uint8;
+struct FEmitterLocalContext;
 
 struct FCodeText
 {
@@ -35,6 +36,26 @@ struct FCodeText
 	}
 };
 
+/** 
+ * A helper struct that adds indented {} scope to the specified context. 
+ * Additionally, manages locals that were added via FEmitHelper::GenerateGetPropertyByName()
+ * and removes them from the cache on destruction, so we don't try to use the
+ * variables outside of the scope.
+ */
+struct FScopeBlock
+{
+public:
+	 FScopeBlock(FEmitterLocalContext& Context);
+	~FScopeBlock();
+
+	void TrackLocalAccessorDecl(const UProperty* Property);
+
+private:
+	FEmitterLocalContext& Context;
+	FScopeBlock*  OuterScopeBlock;
+	TArray<const UProperty*> LocalAccessorDecls;
+};
+
 struct FEmitterLocalContext
 {
 	enum class EClassSubobjectList
@@ -66,6 +87,8 @@ struct FEmitterLocalContext
 	//ConstructorOnly Local Names
 	TMap<UObject*, FString> CommonSubobjectsMap;
 
+	// used to track the innermost FScopeBlock on the stack, so GenerateGetPropertyByName() can register added locals
+	FScopeBlock* ActiveScopeBlock;
 
 private:
 	int32 LocalNameIndexMax;
@@ -88,6 +111,7 @@ public:
 
 	FEmitterLocalContext(const FGatherConvertedClassDependencies& InDependencies)
 		: CurrentCodeType(EGeneratedCodeType::Regular)
+		, ActiveScopeBlock(nullptr)
 		, LocalNameIndexMax(0)
 		, DefaultTarget(&Body)
 		, Dependencies(InDependencies)
@@ -314,7 +338,7 @@ struct FEmitDefaultValueHelper
 	static bool SpecialStructureConstructor(const UStruct* Struct, const uint8* ValuePtr, FString* OutResult);
 
 	// Add static initialization functions. Must be called after Context.UsedObjectInCurrentClass is fully filled
-	static void AddStaticFunctionsForDependencies(FEmitterLocalContext& Context, TSharedPtr<FGatherConvertedClassDependencies> ParentDependencies);
+	static void AddStaticFunctionsForDependencies(FEmitterLocalContext& Context, TSharedPtr<FGatherConvertedClassDependencies> ParentDependencies, FCompilerNativizationOptions NativizationOptions);
 
 	static void AddRegisterHelper(FEmitterLocalContext& Context);
 private:
diff --git a/Engine/Source/Developer/BlueprintCompilerCppBackend/Private/BlueprintCompilerCppBackendValueHelper.cpp b/Engine/Source/Developer/BlueprintCompilerCppBackend/Private/BlueprintCompilerCppBackendValueHelper.cpp
--- a/Engine/Source/Developer/BlueprintCompilerCppBackend/Private/BlueprintCompilerCppBackendValueHelper.cpp
+++ b/Engine/Source/Developer/BlueprintCompilerCppBackend/Private/BlueprintCompilerCppBackendValueHelper.cpp
@@ -193,10 +193,14 @@ void FEmitDefaultValueHelper::InnerGenerate(FEmitterLocalContext& Context, const
 		return bComplete;
 	};
 
-	auto StructProperty = Cast<const UStructProperty>(Property);
+	const UStructProperty* StructProperty = Cast<const UStructProperty>(Property);
 	check(!StructProperty || StructProperty->Struct);
-	auto ArrayProperty = Cast<const UArrayProperty>(Property);
+	const UArrayProperty* ArrayProperty = Cast<const UArrayProperty>(Property);
 	check(!ArrayProperty || ArrayProperty->Inner);
+	const USetProperty* SetProperty = Cast<const USetProperty>(Property);
+	check(!SetProperty || SetProperty->ElementProp);
+	const UMapProperty* MapProperty = Cast<const UMapProperty>(Property);
+	check(!MapProperty || (MapProperty->KeyProp && MapProperty->ValueProp));
 
 	if (!bWithoutFirstConstructionLine)
 	{
@@ -207,7 +211,7 @@ void FEmitDefaultValueHelper::InnerGenerate(FEmitterLocalContext& Context, const
 			Context.AddLine(FString::Printf(TEXT("%s = %s;"), *PathToMember, *ValueStr));
 		}
 		// array initialization "array_var = TArray<..>()" is complete, but it still needs items.
-		if (bComplete && !ArrayProperty)
+		if (bComplete && !ArrayProperty && !SetProperty && !MapProperty)
 		{
 			return;
 		}
@@ -229,63 +233,233 @@ void FEmitDefaultValueHelper::InnerGenerate(FEmitterLocalContext& Context, const
 	if (ArrayProperty)
 	{
 		FScriptArrayHelper ScriptArrayHelper(ArrayProperty, ValuePtr);
-		UStructProperty* InnerStructProperty = Cast<UStructProperty>(ArrayProperty->Inner);
-		
-		const bool bNoExportStructure = InnerStructProperty && InnerStructProperty->Struct
-			&& InnerStructProperty->Struct->IsNative()
-			&& (InnerStructProperty->Struct->StructFlags & STRUCT_NoExport);
-		const bool bInitializeWithoutScriptStruct = bNoExportStructure;
-
-		UScriptStruct* RegularInnerStruct = nullptr;
-		if (!bInitializeWithoutScriptStruct)
+		if (ScriptArrayHelper.Num())
 		{
-			if (InnerStructProperty && !FEmitDefaultValueHelper::SpecialStructureConstructor(InnerStructProperty->Struct, nullptr, nullptr))
+			const UStructProperty* InnerStructProperty = Cast<const UStructProperty>(ArrayProperty->Inner);
+			const bool bInitializeWithoutScriptStruct =  InnerStructProperty && InnerStructProperty->Struct
+				&& InnerStructProperty->Struct->IsNative()
+				&& (InnerStructProperty->Struct->StructFlags & STRUCT_NoExport);
+			const UScriptStruct* RegularInnerStruct = nullptr;
+			if(!bInitializeWithoutScriptStruct)
 			{
-				RegularInnerStruct = InnerStructProperty->Struct;
+				if (InnerStructProperty && !FEmitDefaultValueHelper::SpecialStructureConstructor(InnerStructProperty->Struct, nullptr, nullptr))
+				{
+					RegularInnerStruct = InnerStructProperty->Struct;
+				}
 			}
-		}
-
-		if (ScriptArrayHelper.Num())
-		{
-			const TCHAR* ArrayReserveFunctionName = RegularInnerStruct ? TEXT("AddUninitialized") : TEXT("Reserve");
-			Context.AddLine(FString::Printf(TEXT("%s.%s(%d);"), *PathToMember, ArrayReserveFunctionName, ScriptArrayHelper.Num()));
-
-			if (RegularInnerStruct)
+			
+			if(RegularInnerStruct)
 			{
-				const FString InnerStructStr = Context.FindGloballyMappedObject(RegularInnerStruct, UScriptStruct::StaticClass());
+				Context.AddLine(FString::Printf(TEXT("%s.%s(%d);"), *PathToMember, TEXT("AddUninitialized"), ScriptArrayHelper.Num()));
 				Context.AddLine(FString::Printf(TEXT("%s->%s(%s.GetData(), %d);")
-					, *InnerStructStr
+					, *Context.FindGloballyMappedObject(RegularInnerStruct, UScriptStruct::StaticClass())
 					, GET_FUNCTION_NAME_STRING_CHECKED(UStruct, InitializeStruct)
 					, *PathToMember
 					, ScriptArrayHelper.Num()));
+
+				const FStructOnScope DefaultStruct(RegularInnerStruct);
+				for (int32 Index = 0; Index < ScriptArrayHelper.Num(); ++Index)
+				{
+					const FString ArrayElementRefName = Context.GenerateUniqueLocalName();
+					Context.AddLine(FString::Printf(TEXT("auto& %s = %s[%d];"), *ArrayElementRefName, *PathToMember, Index));
+					InnerGenerate(Context, ArrayProperty->Inner, ArrayElementRefName, ScriptArrayHelper.GetRawPtr(Index), nullptr, true);
+				}
+			}
+			else
+			{
+				Context.AddLine(FString::Printf(TEXT("%s.%s(%d);"), *PathToMember, TEXT("Reserve"), ScriptArrayHelper.Num()));
+				
+				for (int32 Index = 0; Index < ScriptArrayHelper.Num(); ++Index)
+				{
+					const uint8* LocalValuePtr = ScriptArrayHelper.GetRawPtr(Index);
+
+					FString ValueStr;
+					bool bComplete = OneLineConstruction(Context, ArrayProperty->Inner, LocalValuePtr, ValueStr, bInitializeWithoutScriptStruct);
+					ensure(bComplete || bInitializeWithoutScriptStruct);
+					Context.AddLine(FString::Printf(TEXT("%s.Add(%s);"), *PathToMember, *ValueStr));
+				
+					if (bInitializeWithoutScriptStruct && !bComplete)
+					{
+						InnerGenerate(Context, ArrayProperty->Inner, FString::Printf(TEXT("%s[%d]"), *PathToMember, Index), LocalValuePtr, nullptr, true);
+					}
+				}
 			}
 		}
+	}
+	else if(SetProperty)
+	{
+		FScriptSetHelper ScriptSetHelper(SetProperty, ValuePtr);
+		
+		if (ScriptSetHelper.Num())
+		{
+			const UStructProperty* ElementStructProperty = Cast<const UStructProperty>(SetProperty->ElementProp);
+			
+			Context.AddLine(FString::Printf(TEXT("%s.Reserve(%d);"), *PathToMember, ScriptSetHelper.Num()));
 
-		const FStructOnScope DefaultStruct(RegularInnerStruct);
+			if(ElementStructProperty)
+			{
+				// declare a temporary scratch struct for the purpose of initialization. We create a scope so that 
+				// we don't have to generate a new name for it every time we initialize a TSet:
+				Context.AddLine(TEXT("{"));
+				Context.IncreaseIndent();
 
-		for (int32 Index = 0; Index < ScriptArrayHelper.Num(); ++Index)
-		{
-			const uint8* LocalValuePtr = ScriptArrayHelper.GetRawPtr(Index);
+				Context.AddLine(FString::Printf(TEXT("%s SetMemberScratchInitializer;"), *FEmitHelper::GetCppName(ElementStructProperty->Struct)));
+
+				int32 Size = ScriptSetHelper.Num();
+				for (int32 I = 0; Size; ++I)
+				{
+					if(ScriptSetHelper.IsValidIndex(I))
+					{
+						--Size;
+
+						// populate SetMemberScratchInitializer:
+						InnerGenerate(Context, SetProperty->ElementProp, TEXT("SetMemberScratchInitializer"), ScriptSetHelper.GetElementPtr(I), nullptr, true);
 
-			bool bComplete = false;
-			if (!RegularInnerStruct)
+						// Move into the TSet:
+						Context.AddLine(FString::Printf(TEXT("%s.Add(MoveTemp(SetMemberScratchInitializer));"), *PathToMember));
+
+						// reset SetMemberScratchInitializer:
+						Context.AddLine(FString::Printf(TEXT("SetMemberScratchInitializer = %s();"), *FEmitHelper::GetCppName(ElementStructProperty->Struct)));
+					}
+				}
+
+				Context.DecreaseIndent();
+				Context.AddLine(TEXT("}"));
+			}
+			else
 			{
-				FString ValueStr;
-				bComplete = OneLineConstruction(Context, ArrayProperty->Inner, LocalValuePtr, ValueStr, bInitializeWithoutScriptStruct);
-				ensure(bComplete || bInitializeWithoutScriptStruct);
-				Context.AddLine(FString::Printf(TEXT("%s.Add(%s);"), *PathToMember, *ValueStr));
+				int32 Size = ScriptSetHelper.Num();
+				for (int32 I = 0; Size; ++I)
+				{
+					if(ScriptSetHelper.IsValidIndex(I))
+					{
+						--Size;
+
+						FString ValueStr;
+						ensure(OneLineConstruction(Context, SetProperty->ElementProp, ScriptSetHelper.GetElementPtr(I), ValueStr, false));
+						Context.AddLine(FString::Printf(TEXT("%s.Add(%s);"), *PathToMember, *ValueStr));
+					}
+				}
 			}
-			
-			if (RegularInnerStruct || (bInitializeWithoutScriptStruct && !bComplete))
+		}
+	}
+	else if(MapProperty)
+	{
+		FScriptMapHelper ScriptMapHelper(MapProperty, ValuePtr);
+		if (ScriptMapHelper.Num())
+		{	
+			const UStructProperty* KeyStructProperty = Cast<const UStructProperty>(MapProperty->KeyProp);
+			const UStructProperty* ValueStructProperty = Cast<const UStructProperty>(MapProperty->ValueProp);
+
+			Context.AddLine(FString::Printf(TEXT("%s.Reserve(%d);"), *PathToMember, ScriptMapHelper.Num()));
+
+			// Split into 4 cases. They are all similar (add the Key/Value to the Map), but structs are significantly more complicated to initialize than primitives:
+			if(KeyStructProperty && ValueStructProperty)
 			{
-				FString LocalPathToMember = FString::Printf(TEXT("%s[%d]"), *PathToMember, Index);
-				if (RegularInnerStruct)
+				// will add indentation and opening/closing brackets
+				// used to manage locals that were added via FEmitHelper::GenerateGetPropertyByName() 
+				// (removes them from the cache so we don't try to use them outside of the scope)
+				FScopeBlock NestedScopeBlock(Context);
+				
+				Context.AddLine(FString::Printf(TEXT("%s MapKeyScratchInitializer;"), *FEmitHelper::GetCppName(KeyStructProperty->Struct)));
+				Context.AddLine(FString::Printf(TEXT("%s MapValueScratchInitializer;"), *FEmitHelper::GetCppName(ValueStructProperty->Struct)));
+
+				int32 Size = ScriptMapHelper.Num();
+				for (int32 I = 0; Size; ++I)
 				{
-					const FString ArrayElementRefName = Context.GenerateUniqueLocalName();
-					Context.AddLine(FString::Printf(TEXT("auto& %s = %s;"), *ArrayElementRefName, *LocalPathToMember));
-					LocalPathToMember = ArrayElementRefName;
+					if(ScriptMapHelper.IsValidIndex(I))
+					{
+						--Size;
+
+						// populate MapKeyScratchInitializer and MapValueScratchInitializer:
+						InnerGenerate(Context, MapProperty->KeyProp, TEXT("MapKeyScratchInitializer"), ScriptMapHelper.GetKeyPtr(I), nullptr, true);
+						InnerGenerate(Context, MapProperty->ValueProp, TEXT("MapValueScratchInitializer"), ScriptMapHelper.GetValuePtr(I), nullptr, true);
+
+						// Move into the TMap:
+						Context.AddLine(FString::Printf(TEXT("%s.Add(MoveTemp(MapKeyScratchInitializer), MoveTemp(MapValueScratchInitializer));"), *PathToMember));
+
+						// reset MapKeyScratchInitializer and MapValueScratchInitializer:
+						Context.AddLine(FString::Printf(TEXT("MapKeyScratchInitializer = %s();"), *FEmitHelper::GetCppName(KeyStructProperty->Struct)));
+						Context.AddLine(FString::Printf(TEXT("MapValueScratchInitializer = %s();"), *FEmitHelper::GetCppName(ValueStructProperty->Struct)));
+					}
+				}
+			}
+			else if( KeyStructProperty )
+			{
+				// will add indentation and opening/closing brackets
+				// used to manage locals that were added via FEmitHelper::GenerateGetPropertyByName() 
+				// (removes them from the cache so we don't try to use them outside of the scope)
+				FScopeBlock NestedScopeBlock(Context);
+				
+				Context.AddLine(FString::Printf(TEXT("%s MapKeyScratchInitializer;"), *FEmitHelper::GetCppName(KeyStructProperty->Struct)));
+
+				int32 Size = ScriptMapHelper.Num();
+				for (int32 I = 0; Size; ++I)
+				{
+					if(ScriptMapHelper.IsValidIndex(I))
+					{
+						--Size;
+
+						// populate MapKeyScratchInitializer:
+						InnerGenerate(Context, MapProperty->KeyProp, TEXT("MapKeyScratchInitializer"), ScriptMapHelper.GetKeyPtr(I), nullptr, true);
+						
+						FString ValueStr;
+						ensure(OneLineConstruction(Context, MapProperty->ValueProp, ScriptMapHelper.GetValuePtr(I), ValueStr, false));
+
+						// Move into the TMap:
+						Context.AddLine(FString::Printf(TEXT("%s.Add(MoveTemp(MapKeyScratchInitializer), %s);"), *PathToMember, *ValueStr));
+
+						// reset MapKeyScratchInitializer:
+						Context.AddLine(FString::Printf(TEXT("MapKeyScratchInitializer = %s();"), *FEmitHelper::GetCppName(KeyStructProperty->Struct)));
+					}
+				}
+			}
+			else if( ValueStructProperty )
+			{
+				// will add indentation and opening/closing brackets
+				// used to manage locals that were added via FEmitHelper::GenerateGetPropertyByName() 
+				// (removes them from the cache so we don't try to use them outside of the scope)
+				FScopeBlock NestedScopeBlock(Context);
+				
+				Context.AddLine(FString::Printf(TEXT("%s MapValueScratchInitializer;"), *FEmitHelper::GetCppName(ValueStructProperty->Struct)));
+
+				int32 Size = ScriptMapHelper.Num();
+				for (int32 I = 0; Size; ++I)
+				{
+					if(ScriptMapHelper.IsValidIndex(I))
+					{
+						--Size;
+
+						// populate MapValueScratchInitializer:
+						InnerGenerate(Context, MapProperty->ValueProp, TEXT("MapValueScratchInitializer"), ScriptMapHelper.GetValuePtr(I), nullptr, true);
+						
+						FString KeyStr;
+						ensure(OneLineConstruction(Context, MapProperty->KeyProp, ScriptMapHelper.GetKeyPtr(I), KeyStr, false));
+
+						// Move into the TMap:
+						Context.AddLine(FString::Printf(TEXT("%s.Add(%s, MoveTemp(MapValueScratchInitializer));"), *PathToMember, *KeyStr));
+
+						// reset MapValueScratchInitializer:
+						Context.AddLine(FString::Printf(TEXT("MapValueScratchInitializer = %s();"), *FEmitHelper::GetCppName(ValueStructProperty->Struct)));
+					}
+				}
+			}
+			else
+			{
+				int32 Size = ScriptMapHelper.Num();
+				for (int32 I = 0; Size; ++I)
+				{
+					if(ScriptMapHelper.IsValidIndex(I))
+					{
+						--Size;
+
+						FString KeyStr;
+						ensure(OneLineConstruction(Context, MapProperty->KeyProp, ScriptMapHelper.GetKeyPtr(I), KeyStr, false));
+						FString ValueStr;
+						ensure(OneLineConstruction(Context, MapProperty->ValueProp, ScriptMapHelper.GetValuePtr(I), ValueStr, false));
+						Context.AddLine(FString::Printf(TEXT("%s.Add(%s, %s);"), *PathToMember, *KeyStr, *ValueStr));
+					}
 				}
-				InnerGenerate(Context, ArrayProperty->Inner, LocalPathToMember, LocalValuePtr, DefaultStruct.GetStructMemory(), true);
 			}
 		}
 	}
@@ -804,71 +978,69 @@ struct FFakeImportTableHelper
 		{
 			auto GatherDependencies = [&](UClass* InClass)
 			{
-				check(InClass);
+				SerializeBeforeSerializeClassDependencies.Add(InClass->GetSuperClass());
+
+				for (const FImplementedInterface& ImplementedInterface : InClass->Interfaces)
 				{
-					SerializeBeforeSerializeClassDependencies.Add(InClass->GetSuperClass());
+					SerializeBeforeSerializeClassDependencies.Add(ImplementedInterface.Class);
+				}
 
-					for (const FImplementedInterface& ImplementedInterface : InClass->Interfaces)
+				TArray<UObject*> ObjectsInsideClass;
+				GetObjectsWithOuter(InClass, ObjectsInsideClass, true);
+				for (UObject* Obj : ObjectsInsideClass)
+				{
+					UProperty* Property = Cast<UProperty>(Obj);
+					if (!Property)
 					{
-						SerializeBeforeSerializeClassDependencies.Add(ImplementedInterface.Class);
+						continue;
 					}
-
-					TArray<UObject*> ObjectsInsideClass;
-					GetObjectsWithOuter(InClass, ObjectsInsideClass, true);
-					for (UObject* Obj : ObjectsInsideClass)
+					const UProperty* OwnerProperty = Property->GetOwnerProperty();
+					if (!ensure(IsValid(OwnerProperty)))
 					{
-						UProperty* Property = Cast<UProperty>(Obj);
-						if (!Property)
-						{
-							continue;
-						}
-						const UProperty* OwnerProperty = Property->GetOwnerProperty();
-						if (!ensure(IsValid(OwnerProperty)))
-						{
-							continue;
-						}
+						continue;
+					}
 						
-						// TODO:
-						// Let UDS_A contain UDS_B. Let UDS_B contain an array or a set of UDS_A. It causes a cyclic dependency. 
-						// Should we try to fix it at this stage?
+					// TODO:
+					// Let UDS_A contain UDS_B. Let UDS_B contain an array or a set of UDS_A. It causes a cyclic dependency. 
+					// Should we try to fix it at this stage?
 
-						const bool bIsParam = (0 != (OwnerProperty->PropertyFlags & CPF_Parm)) && OwnerProperty->IsIn(InClass);
-						const bool bIsMemberVariable = (OwnerProperty->GetOuter() == InClass);
-						if (bIsParam || bIsMemberVariable) // Affects the class signature. It is necessary while ZCOnstructor/linking.
+					const bool bIsParam = (0 != (OwnerProperty->PropertyFlags & CPF_Parm)) && OwnerProperty->IsIn(InClass);
+					const bool bIsMemberVariable = (OwnerProperty->GetOuter() == InClass);
+					if (bIsParam || bIsMemberVariable) // Affects the class signature. It is necessary while ZCOnstructor/linking.
+					{
+						TArray<UObject*> LocalPreloadDependencies;
+						Property->GetPreloadDependencies(LocalPreloadDependencies);
+						for (UObject* Dependency : LocalPreloadDependencies)
 						{
-							TArray<UObject*> LocalPreloadDependencies;
-							Property->GetPreloadDependencies(LocalPreloadDependencies);
-							for (UObject* Dependency : LocalPreloadDependencies)
+							const bool bDependencyMustBeSerializedBeforeClassIsLinked = Dependency
+								&& (Dependency->IsA<UScriptStruct>() || Dependency->IsA<UEnum>());
+							if (bDependencyMustBeSerializedBeforeClassIsLinked)
 							{
-								const bool bDependencyMustBeSerializedBeforeClassIsLinked = Dependency
-									&& (Dependency->IsA<UScriptStruct>() || Dependency->IsA<UEnum>());
-								if (bDependencyMustBeSerializedBeforeClassIsLinked)
-								{
-									SerializeBeforeSerializeClassDependencies.Add(Dependency);
-								}
+								SerializeBeforeSerializeClassDependencies.Add(Dependency);
 							}
 						}
 					}
 				}
+				SerializeBeforeCreateCDODependencies.Add(InClass->GetSuperClass()->GetDefaultObject());
+			};
+
+			GatherDependencies(SourceClass);
+			GatherDependencies(OriginalClass);
 
+			auto GetClassesOfSubobjects = [&](TMap<UObject*, FString>& SubobjectsMap)
+			{
+				TArray<UObject*> Subobjects;
+				SubobjectsMap.GetKeys(Subobjects);
+				for (UObject* Subobject : Subobjects)
 				{
-					SerializeBeforeCreateCDODependencies.Add(InClass->GetSuperClass()->GetDefaultObject());
-					auto GetClassesOfSubobjects = [&](TMap<UObject*, FString>& SubobjectsMap)
+					if (Subobject)
 					{
-						TArray<UObject*> Subobjects;
-						SubobjectsMap.GetKeys(Subobjects);
-						for (UObject* Subobject : Subobjects)
-						{
-							if (Subobject)
-							{
-								SerializeBeforeCreateCDODependencies.Add(Subobject->GetClass());
-							}
-						}
-					};
-					GetClassesOfSubobjects(Context.ClassSubobjectsMap);
-					GetClassesOfSubobjects(Context.CommonSubobjectsMap);
+						SerializeBeforeCreateCDODependencies.Add(Subobject->GetClass());
+					}
 				}
 			};
+			GetClassesOfSubobjects(Context.ClassSubobjectsMap);
+			GetClassesOfSubobjects(Context.CommonSubobjectsMap);
 		}
 	}
 
@@ -904,7 +1076,9 @@ struct FFakeImportTableHelper
 	}
 };
 
-void FEmitDefaultValueHelper::AddStaticFunctionsForDependencies(FEmitterLocalContext& Context, TSharedPtr<FGatherConvertedClassDependencies> ParentDependencies)
+void FEmitDefaultValueHelper::AddStaticFunctionsForDependencies(FEmitterLocalContext& Context
+	, TSharedPtr<FGatherConvertedClassDependencies> ParentDependencies
+	, FCompilerNativizationOptions NativizationOptions)
 {
 	// 1. GATHER UDS DEFAULT VALUE DEPENDENCIES
 	{
@@ -922,7 +1096,18 @@ void FEmitDefaultValueHelper::AddStaticFunctionsForDependencies(FEmitterLocalCon
 	// 2. ALL ASSETS TO LIST
 	TSet<const UObject*> AllDependenciesToHandle = Context.Dependencies.AllDependencies();
 	AllDependenciesToHandle.Append(Context.UsedObjectInCurrentClass);
-	
+	AllDependenciesToHandle.Remove(nullptr);
+
+	// Special case, we don't need to load any dependencies from CoreUObject.
+	UPackage* CoreUObjectPackage = UProperty::StaticClass()->GetOutermost();
+	for (auto Iter = AllDependenciesToHandle.CreateIterator(); Iter; ++Iter)
+	{
+		if ((*Iter)->GetOutermost() == CoreUObjectPackage)
+		{
+			Iter.RemoveCurrent();
+		}
+	}
+
 	// HELPERS
 	auto SourceClass = Context.GetCurrentlyGeneratedClass();
 	auto OriginalClass = Context.Dependencies.FindOriginalClass(SourceClass);
@@ -954,9 +1139,38 @@ void FEmitDefaultValueHelper::AddStaticFunctionsForDependencies(FEmitterLocalCon
 			, *AssetType->GetName());
 	};
 
-	auto CreateDependencyRecord = [&](const UObject* InAsset) -> FCompactBlueprintDependencyData
+	auto CreateDependencyRecord = [&](const UObject* InAsset, FString& OptionalComment) -> FCompactBlueprintDependencyData
 	{
 		ensure(InAsset);
+		if (InAsset && IsEditorOnlyObject(InAsset))
+		{
+			UE_LOG(LogK2Compiler, Warning, TEXT("Nativized %d depends on editor only asset: %s")
+				, (OriginalClass ? *OriginalClass->GetPathName() : *CppClassName)
+				,*InAsset->GetPathName());
+			OptionalComment = TEXT("Editor Only asset");
+			return FCompactBlueprintDependencyData{};
+		}
+
+		{
+			bool bNotForClient = false;
+			bool bNotForServer = false;
+			for (const UObject* Search = InAsset; Search && !Search->IsA(UPackage::StaticClass()); Search = Search->GetOuter())
+			{
+				bNotForClient = bNotForClient || !Search->NeedsLoadForClient();
+				bNotForServer = bNotForServer || !Search->NeedsLoadForServer();
+			}
+			if (bNotForServer && NativizationOptions.ServerOnlyPlatform)
+			{
+				OptionalComment = TEXT("Not for server");
+				return FCompactBlueprintDependencyData{};
+			}
+			if (bNotForClient && NativizationOptions.ClientOnlyPlatform)
+			{
+				OptionalComment = TEXT("Not for client");
+				return FCompactBlueprintDependencyData{};
+			}
+		}
+
 		FNativizationSummary::FDependencyRecord& DependencyRecord = FDependenciesGlobalMapHelper::FindDependencyRecord(InAsset);
 		ensure(DependencyRecord.Index >= 0);
 		if (DependencyRecord.NativeLine.IsEmpty())
@@ -990,11 +1204,13 @@ void FEmitDefaultValueHelper::AddStaticFunctionsForDependencies(FEmitterLocalCon
 
 		for (const UObject* LocAsset : Assets)
 		{
-			const FCompactBlueprintDependencyData DependencyRecord = CreateDependencyRecord(LocAsset);
-			Context.AddLine(FString::Printf(TEXT("{%d, %s, %s},  // %s ")
+			FString OptionalComment;
+			const FCompactBlueprintDependencyData DependencyRecord = CreateDependencyRecord(LocAsset, OptionalComment);
+			Context.AddLine(FString::Printf(TEXT("{%d, %s, %s},  // %s %s ")
 				, DependencyRecord.ObjectRefIndex
 				, *BlueprintDependencyTypeToString(DependencyRecord.ClassDependency)
 				, *BlueprintDependencyTypeToString(DependencyRecord.CDODependency)
+				, *OptionalComment
 				, *LocAsset->GetFullName()));
 		}
 
@@ -1059,7 +1275,7 @@ void FEmitDefaultValueHelper::AddRegisterHelper(FEmitterLocalContext& Context)
 	Context.IncreaseIndent();
 
 	Context.AddLine(FString::Printf(
-		TEXT("FConvertedBlueprintsDependencies::Get().RegisterClass(TEXT(\"%s\"), &%s::__StaticDependenciesAssets);")
+		TEXT("FConvertedBlueprintsDependencies::Get().RegisterConvertedClass(TEXT(\"%s\"), &%s::__StaticDependenciesAssets);")
 		, *OriginalClass->GetOutermost()->GetPathName()
 		, *CppClassName));
 
@@ -1145,6 +1361,15 @@ void FEmitDefaultValueHelper::GenerateCustomDynamicClassInitialization(FEmitterL
 	// Gather all CT from SCS and IH, the remaining ones are generated for class..
 	if (auto SCS = BPGC->SimpleConstructionScript)
 	{
+		// >>> This code should be removed, once UE-39168 is fixed
+		//TODO: it's an ugly workaround - template from DefaultSceneRootNode is unnecessarily cooked :(
+		UActorComponent* DefaultSceneRootComponentTemplate = SCS->GetDefaultSceneRootNode() ? SCS->GetDefaultSceneRootNode()->ComponentTemplate : nullptr;
+		if (DefaultSceneRootComponentTemplate)
+		{
+			ActorComponentTempatesOwnedByClass.Add(DefaultSceneRootComponentTemplate);
+		}
+		// <<< This code should be removed, once UE-39168 is fixed
+
 		for (auto Node : SCS->GetAllNodes())
 		{
 			ActorComponentTempatesOwnedByClass.RemoveSwap(Node->ComponentTemplate);
@@ -1411,6 +1636,11 @@ FString FEmitDefaultValueHelper::HandleClassSubobject(FEmitterLocalContext& Cont
 		UClass* ObjectClass = Object->GetClass();
 		const FString ActualClass = Context.FindGloballyMappedObject(ObjectClass, UClass::StaticClass());
 		const FString NativeType = FEmitHelper::GetCppName(Context.GetFirstNativeOrConvertedClass(ObjectClass));
+		if(!ObjectClass->IsNative())
+		{
+			// make sure CDO has been created for NativeType:
+			Context.AddLine(FString::Printf(TEXT("%s::StaticClass()->GetDefaultObject();"), *NativeType));
+		}
 		Context.AddLine(FString::Printf(
 			TEXT("auto %s = NewObject<%s>(%s, %s, TEXT(\"%s\"));")
 			, *LocalNativeName
@@ -1525,6 +1755,11 @@ FString FEmitDefaultValueHelper::HandleInstancedSubobject(FEmitterLocalContext&
 
 		const FString ActualClass = Context.FindGloballyMappedObject(ObjectClass, UClass::StaticClass());
 		const FString NativeType = FEmitHelper::GetCppName(Context.GetFirstNativeOrConvertedClass(ObjectClass));
+		if(!ObjectClass->IsNative())
+		{
+			// make sure CDO has been created for NativeType:
+			Context.AddLine(FString::Printf(TEXT("%s::StaticClass()->GetDefaultObject();"), *NativeType));
+		}
 		Context.AddLine(FString::Printf(
 			TEXT("auto %s = NewObject<%s>(%s, %s, TEXT(\"%s\"));")
 			, *LocalNativeName
diff --git a/Engine/Source/Developer/BlueprintNativeCodeGen/Private/BlueprintNativeCodeGenManifest.cpp b/Engine/Source/Developer/BlueprintNativeCodeGen/Private/BlueprintNativeCodeGenManifest.cpp
--- a/Engine/Source/Developer/BlueprintNativeCodeGen/Private/BlueprintNativeCodeGenManifest.cpp
+++ b/Engine/Source/Developer/BlueprintNativeCodeGen/Private/BlueprintNativeCodeGenManifest.cpp
@@ -400,7 +400,8 @@ FBlueprintNativeCodeGenManifest::FBlueprintNativeCodeGenManifest()
 }
 
 //------------------------------------------------------------------------------
-FBlueprintNativeCodeGenManifest::FBlueprintNativeCodeGenManifest(const FString& InPluginName, const FString& InOutputDir)
+FBlueprintNativeCodeGenManifest::FBlueprintNativeCodeGenManifest(const FString& InPluginName, const FString& InOutputDir, FCompilerNativizationOptions InCompilerNativizationOptions)
+	: NativizationOptions(InCompilerNativizationOptions)
 {
 	using namespace BlueprintNativeCodeGenManifestImpl;
 	
diff --git a/Engine/Source/Developer/BlueprintNativeCodeGen/Private/BlueprintNativeCodeGenManifest.h b/Engine/Source/Developer/BlueprintNativeCodeGen/Private/BlueprintNativeCodeGenManifest.h
--- a/Engine/Source/Developer/BlueprintNativeCodeGen/Private/BlueprintNativeCodeGenManifest.h
+++ b/Engine/Source/Developer/BlueprintNativeCodeGen/Private/BlueprintNativeCodeGenManifest.h
@@ -118,7 +118,7 @@ struct FBlueprintNativeCodeGenManifest
 
 public: 
 	FBlueprintNativeCodeGenManifest(); // default ctor for USTRUCT() system
-	FBlueprintNativeCodeGenManifest(const FString& InPluginName, const FString& InOutputDir);
+	FBlueprintNativeCodeGenManifest(const FString& InPluginName, const FString& InOutputDir, FCompilerNativizationOptions InCompilerNativizationOptions);
 	FBlueprintNativeCodeGenManifest(const FString& ManifestFilePath);
 
 	/**
@@ -168,6 +168,11 @@ public:
 	*/
 	const FUnconvertedRecord& GetUnconvertedDependencies() const { return UnconvertedDependencies; }
 
+	/**
+	 * @return compiler nativization options
+	 */
+	const FCompilerNativizationOptions GetCompilerNativizationOptions() const { return NativizationOptions; }
+
 	/**
 	 * Saves this manifest as json, to its target destination (which it was 
 	 * setup with).
@@ -214,4 +219,7 @@ private:
 
 	UPROPERTY()
 	TMap<FName, FUnconvertedDependencyRecord> UnconvertedDependencies;
+
+	UPROPERTY()
+	FCompilerNativizationOptions NativizationOptions;
 };
diff --git a/Engine/Source/Developer/BlueprintNativeCodeGen/Private/BlueprintNativeCodeGenModule.cpp b/Engine/Source/Developer/BlueprintNativeCodeGen/Private/BlueprintNativeCodeGenModule.cpp
--- a/Engine/Source/Developer/BlueprintNativeCodeGen/Private/BlueprintNativeCodeGenModule.cpp
+++ b/Engine/Source/Developer/BlueprintNativeCodeGen/Private/BlueprintNativeCodeGenModule.cpp
@@ -65,7 +65,7 @@ public:
 	}
 protected:
 	virtual void Initialize(const FNativeCodeGenInitData& InitData) override;
-	virtual void InitializeForRerunDebugOnly(const TArray< TPair< FString, FString > >& CodegenTargets) override;
+	virtual void InitializeForRerunDebugOnly(const TArray<FPlatformNativizationDetails>& CodegenTargets) override;
 	//~ End IBlueprintNativeCodeGenModule interface
 
 	//~ Begin FScriptCookReplacmentCoordinator interface
@@ -102,6 +102,9 @@ private:
 
 	TSet<TAssetPtr<UBlueprint>> ToGenerate;
 	TMultiMap<FName, TAssetSubclassOf<UObject>> FunctionsBoundToADelegate; // is a function could be bound to a delegate, then it must have UFUNCTION macro. So we cannot optimize it.
+
+	// Cached values from IsTargetedForReplacement
+	mutable TMap<FStringAssetReference, EReplacementResult> CachedIsTargetedForReplacement;
 };
 
 void FBlueprintNativeCodeGenModule::ReadConfig()
@@ -234,22 +237,22 @@ void FBlueprintNativeCodeGenModule::FillIsFunctionUsedInADelegate()
 
 void FBlueprintNativeCodeGenModule::Initialize(const FNativeCodeGenInitData& InitData)
 {
+	CachedIsTargetedForReplacement.Reset();
 	ReadConfig();
 
 	IBlueprintNativeCodeGenCore::Register(this);
 
 	// Each platform will need a manifest, because each platform could cook different assets:
-	for (auto& Platform : InitData.CodegenTargets)
+	for (const FPlatformNativizationDetails& Platform : InitData.CodegenTargets)
 	{
-		const TCHAR* TargetDirectory = *Platform.Value;
-		FString OutputPath = FPaths::Combine(TargetDirectory, NativizationCookControllerImpl::DefaultPluginName);
+		FString OutputPath = FPaths::Combine(*Platform.PlatformTargetDirectory, NativizationCookControllerImpl::DefaultPluginName);
 
-		Manifests.Add(FString(*Platform.Key), TUniquePtr<FBlueprintNativeCodeGenManifest>(new FBlueprintNativeCodeGenManifest(NativizationCookControllerImpl::DefaultPluginName, OutputPath)));
+		Manifests.Add(*Platform.PlatformName, TUniquePtr<FBlueprintNativeCodeGenManifest>(new FBlueprintNativeCodeGenManifest(NativizationCookControllerImpl::DefaultPluginName, OutputPath, Platform.CompilerNativizationOptions)));
 
-		TargetPlatformNames.Add(Platform.Key);
+		TargetPlatformNames.Add(Platform.PlatformName);
 
 		// Clear source code folder
-		const FString SourceCodeDir = GetManifest(*Platform.Key).GetTargetPaths().PluginSourceDir();
+		const FString SourceCodeDir = GetManifest(*Platform.PlatformName).GetTargetPaths().PluginSourceDir();
 		UE_LOG(LogBlueprintCodeGen, Log, TEXT("Clear nativized source code directory: %s"), *SourceCodeDir);
 		IFileManager::Get().DeleteDirectory(*SourceCodeDir, false, true);
 	}
@@ -259,22 +262,22 @@ void FBlueprintNativeCodeGenModule::Initialize(const FNativeCodeGenInitData& Ini
 	FillIsFunctionUsedInADelegate();
 }
 
-void FBlueprintNativeCodeGenModule::InitializeForRerunDebugOnly(const TArray< TPair< FString, FString > >& CodegenTargets)
+void FBlueprintNativeCodeGenModule::InitializeForRerunDebugOnly(const TArray<FPlatformNativizationDetails>& CodegenTargets)
 {
+	CachedIsTargetedForReplacement.Reset();
 	ReadConfig();
 	IBlueprintNativeCodeGenCore::Register(this);
 	FillTargetedForReplacementQuery();
 	FillIsFunctionUsedInADelegate();
 
-	for (const auto& Platform : CodegenTargets)
+	for (const FPlatformNativizationDetails& Platform : CodegenTargets)
 	{
 		// load the old manifest:
-		const TCHAR* TargetDirectory = *Platform.Value;
-		FString OutputPath = FPaths::Combine(TargetDirectory, NativizationCookControllerImpl::DefaultPluginName, *FBlueprintNativeCodeGenPaths::GetDefaultManifestPath());
-		Manifests.Add(FString(*Platform.Key), TUniquePtr<FBlueprintNativeCodeGenManifest>(new FBlueprintNativeCodeGenManifest(FPaths::ConvertRelativePathToFull(OutputPath))));
+		FString OutputPath = FPaths::Combine(*Platform.PlatformTargetDirectory, NativizationCookControllerImpl::DefaultPluginName, *FBlueprintNativeCodeGenPaths::GetDefaultManifestPath());
+		Manifests.Add(Platform.PlatformName, TUniquePtr<FBlueprintNativeCodeGenManifest>(new FBlueprintNativeCodeGenManifest(FPaths::ConvertRelativePathToFull(OutputPath))));
 		//FBlueprintNativeCodeGenManifest OldManifest(FPaths::ConvertRelativePathToFull(OutputPath));
 		// reconvert every assets listed in the manifest:
-		for (const auto& ConversionTarget : GetManifest(*Platform.Key).GetConversionRecord())
+		for (const auto& ConversionTarget : GetManifest(*Platform.PlatformName).GetConversionRecord())
 		{
 			// load the package:
 			UPackage* Package = LoadPackage(nullptr, *ConversionTarget.Value.TargetObjPath, LOAD_None);
@@ -286,11 +289,11 @@ void FBlueprintNativeCodeGenModule::InitializeForRerunDebugOnly(const TArray< TP
 			}
 
 			// reconvert it
-			Convert(Package, ESavePackageResult::ReplaceCompletely, *Platform.Key);
+			Convert(Package, ESavePackageResult::ReplaceCompletely, *Platform.PlatformName);
 		}
 
 		// reconvert every unconverted dependency listed in the manifest:
-		for (const auto& ConversionTarget : GetManifest(*Platform.Key).GetUnconvertedDependencies())
+		for (const auto& ConversionTarget : GetManifest(*Platform.PlatformName).GetUnconvertedDependencies())
 		{
 			// load the package:
 			UPackage* Package = LoadPackage(nullptr, *ConversionTarget.Key.GetPlainNameString(), LOAD_None);
@@ -302,7 +305,7 @@ void FBlueprintNativeCodeGenModule::InitializeForRerunDebugOnly(const TArray< TP
 			if (ensure(BP))
 			{
 				CollectBoundFunctions(BP);
-				GenerateSingleStub(BP, *Platform.Key);
+				GenerateSingleStub(BP, *Platform.PlatformName);
 			}
 		}
 
@@ -311,7 +314,7 @@ void FBlueprintNativeCodeGenModule::InitializeForRerunDebugOnly(const TArray< TP
 			UBlueprint* BP = BPPtr.LoadSynchronous();
 			if (ensure(BP))
 			{
-				GenerateSingleAsset(BP->GeneratedClass, *Platform.Key);
+				GenerateSingleAsset(BP->GeneratedClass, *Platform.PlatformName);
 			}
 		}
 	}
@@ -420,7 +423,7 @@ void FBlueprintNativeCodeGenModule::GenerateSingleAsset(UField* ForConversion, c
 	TSharedPtr<FString> HeaderSource(new FString());
 	TSharedPtr<FString> CppSource(new FString());
 
-	FBlueprintNativeCodeGenUtils::GenerateCppCode(ForConversion, HeaderSource, CppSource, NativizationSummary);
+	FBlueprintNativeCodeGenUtils::GenerateCppCode(ForConversion, HeaderSource, CppSource, NativizationSummary, GetManifest(PlatformName).GetCompilerNativizationOptions());
 	bool bSuccess = !HeaderSource->IsEmpty() || !CppSource->IsEmpty();
 	// Run the cpp first, because we cue off of the presence of a header for a valid conversion record (see
 	// FConvertedAssetRecord::IsValid)
@@ -605,7 +608,6 @@ UClass* FBlueprintNativeCodeGenModule::FindReplacedClassForObject(const UObject*
 UObject* FBlueprintNativeCodeGenModule::FindReplacedNameAndOuter(UObject* Object, FName& OutName) const
 {
 	OutName = NAME_None;
-	UObject* Outer = nullptr;
 
 	UActorComponent* ActorComponent = Cast<UActorComponent>(Object);
 	if (ActorComponent)
@@ -616,20 +618,22 @@ UObject* FBlueprintNativeCodeGenModule::FindReplacedNameAndOuter(UObject* Object
 		{
 			if (OuterObject->HasAnyFlags(RF_ClassDefaultObject))
 			{
-				return Outer;
+				return nullptr;
 			}
 			BPGC = Cast<UBlueprintGeneratedClass>(OuterObject);
 		}
 
-		for (UBlueprintGeneratedClass* SuperBPGC = BPGC; SuperBPGC && (OutName == NAME_None); SuperBPGC = Cast<UBlueprintGeneratedClass>(SuperBPGC->GetSuperClass()))
+		FName NewName = NAME_None;
+		UObject* OuterCDO = nullptr;
+		for (UBlueprintGeneratedClass* SuperBPGC = BPGC; SuperBPGC && (NewName == NAME_None); SuperBPGC = Cast<UBlueprintGeneratedClass>(SuperBPGC->GetSuperClass()))
 		{
 			if (SuperBPGC->InheritableComponentHandler)
 			{
 				FComponentKey FoundKey = SuperBPGC->InheritableComponentHandler->FindKey(ActorComponent);
 				if (FoundKey.IsValid())
 				{
-					OutName = FoundKey.IsSCSKey() ? FoundKey.GetSCSVariableName() : ActorComponent->GetFName();
-					Outer = BPGC->GetDefaultObject(false);
+					NewName = FoundKey.IsSCSKey() ? FoundKey.GetSCSVariableName() : ActorComponent->GetFName();
+					OuterCDO = BPGC->GetDefaultObject(false);
 					break;
 				}
 			}
@@ -639,22 +643,23 @@ UObject* FBlueprintNativeCodeGenModule::FindReplacedNameAndOuter(UObject* Object
 				{
 					if (Node->ComponentTemplate == ActorComponent)
 					{
-						OutName = Node->GetVariableName();
-						if (OutName != NAME_None)
+						NewName = Node->GetVariableName();
+						if (NewName != NAME_None)
 						{
-							Outer = BPGC->GetDefaultObject(false);
+							OuterCDO = BPGC->GetDefaultObject(false);
 							break;
 						}
 					}
 				}
 			}
 		}
-	}
 
-	if (Outer && (EReplacementResult::ReplaceCompletely == IsTargetedForReplacement(Object->GetClass())))
-	{
-		UE_LOG(LogBlueprintCodeGen, Log, TEXT("Object '%s' has replaced name '%s' and outer: '%s'"), *GetPathNameSafe(Object), *OutName.ToString(), *GetPathNameSafe(Outer));
-		return Outer;
+		if (OuterCDO && (EReplacementResult::ReplaceCompletely == IsTargetedForReplacement(OuterCDO->GetClass())))
+		{
+			OutName = NewName;
+			UE_LOG(LogBlueprintCodeGen, Log, TEXT("Object '%s' has replaced name '%s' and outer: '%s'"), *GetPathNameSafe(Object), *OutName.ToString(), *GetPathNameSafe(OuterCDO));
+			return OuterCDO;
+		}
 	}
 
 	return nullptr;
@@ -677,26 +682,35 @@ EReplacementResult FBlueprintNativeCodeGenModule::IsTargetedForReplacement(const
 
 EReplacementResult FBlueprintNativeCodeGenModule::IsTargetedForReplacement(const UObject* Object) const
 {
-	const UStruct* const Struct = Cast<UStruct>(Object);
-	const UEnum* const Enum = Cast<UEnum>(Object);
-	const UClass* const BlueprintClass = Cast<UClass>(Struct);
-	const UBlueprint* const Blueprint = BlueprintClass ? Cast<UBlueprint>(BlueprintClass->ClassGeneratedBy) : nullptr;
+	if (Object == nullptr)
+	{
+		return EReplacementResult::DontReplace;
+	}
 
-	auto ObjectIsNotReplacedAtAll = [&]() -> bool
+	const UUserDefinedStruct* const UDStruct = Cast<const UUserDefinedStruct>(Object);
+	const UUserDefinedEnum* const UDEnum = Cast<const UUserDefinedEnum>(Object);
+	const UBlueprintGeneratedClass* const BlueprintClass = Cast<const UBlueprintGeneratedClass>(Object);
+	if (UDStruct == nullptr && UDEnum == nullptr && BlueprintClass == nullptr)
 	{
-		if (Object == nullptr)
-		{
-			return true;
-		}
+		return EReplacementResult::DontReplace;
+	}
 
-		if (Struct == nullptr && Enum == nullptr)
+	const FStringAssetReference ObjectKey(Object);
+	{
+		const EReplacementResult* const CachedValue = CachedIsTargetedForReplacement.Find(ObjectKey); //THe referenced returned by FindOrAdd could be invalid later, when filled.
+		if (CachedValue)
 		{
-			return true;
+			return *CachedValue;
 		}
+	}
 
+	const UBlueprint* const Blueprint = BlueprintClass ? Cast<UBlueprint>(BlueprintClass->ClassGeneratedBy) : nullptr;
+
+	auto ObjectIsNotReplacedAtAll = [&]() -> bool
+	{
 		// EDITOR ON DEVELOPMENT OBJECT
 		{
-			auto IsDeveloperObject = [&](const UObject* Obj) -> bool
+			auto IsDeveloperObject = [](const UObject* Obj) -> bool
 			{
 				auto IsObjectFromDeveloperPackage = [](const UObject* InObj) -> bool
 				{
@@ -734,10 +748,16 @@ EReplacementResult FBlueprintNativeCodeGenModule::IsTargetedForReplacement(const
 				return true;
 			}
 		}
+		// Don't convert objects like Default__WidgetBlueprintGeneratedClass
+		if (Object && (Object->HasAnyFlags(RF_ClassDefaultObject)))
+		{
+			return true;
+		}
 		return false;
 	};
 	if (ObjectIsNotReplacedAtAll())
 	{
+		CachedIsTargetedForReplacement.Add(ObjectKey, EReplacementResult::DontReplace);
 		return EReplacementResult::DontReplace;
 	}
 
@@ -759,12 +779,12 @@ EReplacementResult FBlueprintNativeCodeGenModule::IsTargetedForReplacement(const
 		// ExcludedAssetTypes
 		{
 			// we can't use FindObject, because we may be converting a type while saving
-			if (Enum && ExcludedAssetTypes.Find(Enum->GetPathName()) != INDEX_NONE)
+			if (UDEnum && ExcludedAssetTypes.Find(UDEnum->GetPathName()) != INDEX_NONE)
 			{
 				return true;
 			}
 
-			const UStruct* LocStruct = Struct;
+			const UStruct* LocStruct = Cast<const UStruct>(Object);
 			while (LocStruct)
 			{
 				if (ExcludedAssetTypes.Find(LocStruct->GetPathName()) != INDEX_NONE)
@@ -835,10 +855,10 @@ EReplacementResult FBlueprintNativeCodeGenModule::IsTargetedForReplacement(const
 			}
 
 			const UProjectPackagingSettings* const PackagingSettings = GetDefault<UProjectPackagingSettings>();
-			const bool bNativizeOnlySelectedBPs = PackagingSettings && PackagingSettings->bNativizeOnlySelectedBlueprints;
+			const bool bNativizeOnlySelectedBPs = PackagingSettings && PackagingSettings->BlueprintNativizationMethod == EProjectPackagingBlueprintNativizationMethod::Exclusive;
 
 			// Blueprint is not selected
-			if (bNativizeOnlySelectedBPs && !Blueprint->bNativize)
+			if (bNativizeOnlySelectedBPs && !Blueprint->bSelectedForNativization)
 			{
 				return true;
 			}
@@ -857,14 +877,36 @@ EReplacementResult FBlueprintNativeCodeGenModule::IsTargetedForReplacement(const
 					return true;
 				}
 			}
+
+			// Interface class not converted
+			TArray<UClass*> InterfaceClasses;
+			FBlueprintEditorUtils::FindImplementedInterfaces(Blueprint, false, InterfaceClasses);
+			for (const UClass* InterfaceClassIt : InterfaceClasses)
+			{
+				const UBlueprintGeneratedClass* InterfaceBPGC = Cast<const UBlueprintGeneratedClass>(InterfaceClassIt);
+				if (InterfaceBPGC)
+				{
+					const EReplacementResult InterfaceResult = IsTargetedForReplacement(InterfaceBPGC);
+					if (InterfaceResult != EReplacementResult::ReplaceCompletely)
+					{
+						if (bNativizeOnlySelectedBPs)
+						{
+							UE_LOG(LogBlueprintCodeGen, Error, TEXT("BP %s is selected for nativization, but BP interface class %s is not nativized."), *GetPathNameSafe(Blueprint), *GetPathNameSafe(InterfaceClassIt));
+						}
+						return true;
+					}
+				}
+			}
 		}
 		return false;
 	};
 	if (ObjectGenratesOnlyStub())
 	{
+		CachedIsTargetedForReplacement.Add(ObjectKey, EReplacementResult::GenerateStub);
 		return EReplacementResult::GenerateStub;
 	}
 
+	CachedIsTargetedForReplacement.Add(ObjectKey, EReplacementResult::ReplaceCompletely);
 	return EReplacementResult::ReplaceCompletely;
 }
 
diff --git a/Engine/Source/Developer/BlueprintNativeCodeGen/Private/BlueprintNativeCodeGenUtils.cpp b/Engine/Source/Developer/BlueprintNativeCodeGen/Private/BlueprintNativeCodeGenUtils.cpp
--- a/Engine/Source/Developer/BlueprintNativeCodeGen/Private/BlueprintNativeCodeGenUtils.cpp
+++ b/Engine/Source/Developer/BlueprintNativeCodeGen/Private/BlueprintNativeCodeGenUtils.cpp
@@ -280,7 +280,7 @@ bool FBlueprintNativeCodeGenUtils::FinalizePlugin(const FBlueprintNativeCodeGenM
 }
 
 //------------------------------------------------------------------------------
-void FBlueprintNativeCodeGenUtils::GenerateCppCode(UObject* Obj, TSharedPtr<FString> OutHeaderSource, TSharedPtr<FString> OutCppSource, TSharedPtr<FNativizationSummary> NativizationSummary)
+void FBlueprintNativeCodeGenUtils::GenerateCppCode(UObject* Obj, TSharedPtr<FString> OutHeaderSource, TSharedPtr<FString> OutCppSource, TSharedPtr<FNativizationSummary> NativizationSummary, FCompilerNativizationOptions NativizationOptions)
 {
 	auto UDEnum = Cast<UUserDefinedEnum>(Obj);
 	auto UDStruct = Cast<UUserDefinedStruct>(Obj);
@@ -346,6 +346,7 @@ void FBlueprintNativeCodeGenUtils::GenerateCppCode(UObject* Obj, TSharedPtr<FStr
 			CompileOptions.CompileType = EKismetCompileType::Cpp;
 			CompileOptions.OutCppSourceCode = OutCppSource;
 			CompileOptions.OutHeaderSourceCode = OutHeaderSource;
+			CompileOptions.NativizationOptions = NativizationOptions;
 
 			Compiler.CompileBlueprint(DuplicateBP, CompileOptions, Results);
 
@@ -362,7 +363,7 @@ void FBlueprintNativeCodeGenUtils::GenerateCppCode(UObject* Obj, TSharedPtr<FStr
 		IKismetCompilerInterface& Compiler = FModuleManager::LoadModuleChecked<IKismetCompilerInterface>(KISMET_COMPILER_MODULENAME);
 		if (UDEnum)
 		{
-			*OutHeaderSource = Compiler.GenerateCppCodeForEnum(UDEnum);
+			Compiler.GenerateCppCodeForEnum(UDEnum, *OutHeaderSource, *OutCppSource);
 		}
 		else if (UDStruct)
 		{
diff --git a/Engine/Source/Developer/BlueprintNativeCodeGen/Private/BlueprintNativeCodeGenUtils.h b/Engine/Source/Developer/BlueprintNativeCodeGen/Private/BlueprintNativeCodeGenUtils.h
--- a/Engine/Source/Developer/BlueprintNativeCodeGen/Private/BlueprintNativeCodeGenUtils.h
+++ b/Engine/Source/Developer/BlueprintNativeCodeGen/Private/BlueprintNativeCodeGenUtils.h
@@ -34,7 +34,7 @@ public:
 	 * @param  OutHeaderSource	Output destination for the header source.
 	 * @param  OutCppSource		Output destination for the cpp source.
 	 */
-	static void GenerateCppCode(UObject* Obj, TSharedPtr<FString> OutHeaderSource, TSharedPtr<FString> OutCppSource, TSharedPtr<FNativizationSummary> NativizationSummary = TSharedPtr<FNativizationSummary>());
+	static void GenerateCppCode(UObject* Obj, TSharedPtr<FString> OutHeaderSource, TSharedPtr<FString> OutCppSource, TSharedPtr<FNativizationSummary> NativizationSummary, FCompilerNativizationOptions NativizationOptions);
 
 public: 
 	/** 
diff --git a/Engine/Source/Developer/BlueprintNativeCodeGen/Private/NativeCodeGenerationTool.cpp b/Engine/Source/Developer/BlueprintNativeCodeGen/Private/NativeCodeGenerationTool.cpp
--- a/Engine/Source/Developer/BlueprintNativeCodeGen/Private/NativeCodeGenerationTool.cpp
+++ b/Engine/Source/Developer/BlueprintNativeCodeGen/Private/NativeCodeGenerationTool.cpp
@@ -139,7 +139,7 @@ struct FGeneratedCodeData
 			return false;
 		}
 
-		const int WorkParts = 3 + (4 * DependentObjects.Num()) + (2 * UnconvertedNeededClasses.Num());
+		const int WorkParts = 3 + (4 * DependentObjects.Num());
 		FScopedSlowTask SlowTask(WorkParts, LOCTEXT("GeneratingCppFiles", "Generating C++ files.."));
 		SlowTask.MakeDialog();
 
@@ -153,7 +153,8 @@ struct FGeneratedCodeData
 
 			TSharedPtr<FString> HeaderSource(new FString());
 			TSharedPtr<FString> CppSource(new FString());
-			FBlueprintNativeCodeGenUtils::GenerateCppCode(Obj, HeaderSource, CppSource);
+			TSharedPtr<FNativizationSummary> NativizationSummary(new FNativizationSummary());
+			FBlueprintNativeCodeGenUtils::GenerateCppCode(Obj, HeaderSource, CppSource, NativizationSummary, FCompilerNativizationOptions{});
 			SlowTask.EnterProgressFrame();
 
 			const FString BackendBaseFilename = CodeGenBackend.ConstructBaseFilename(Obj);
@@ -185,73 +186,8 @@ struct FGeneratedCodeData
 			}
 		}
 
-		for (auto BPGC : UnconvertedNeededClasses)
-		{
-			SlowTask.EnterProgressFrame();
-
-			IKismetCompilerInterface& Compiler = FModuleManager::LoadModuleChecked<IKismetCompilerInterface>(KISMET_COMPILER_MODULENAME);
-			const FString HeaderSource = Compiler.GenerateCppWrapper(BPGC);
-
-			SlowTask.EnterProgressFrame();
-
-			const FString BackendBaseFilename = CodeGenBackend.ConstructBaseFilename(BPGC);
-
-			const FString FullHeaderFilename = FPaths::Combine(*HeaderDirPath, *(BackendBaseFilename + TEXT(".h")));
-			const bool bHeaderSaved = FFileHelper::SaveStringToFile(HeaderSource, *FullHeaderFilename);
-			if (!bHeaderSaved)
-			{
-				ErrorString += FString::Printf(*LOCTEXT("HeaderNotSaved", "Header file wasn't saved. Check log for details. %s\n").ToString(), *BPGC->GetPathName());
-			}
-			else
-			{
-				CreatedFiles.Add(FullHeaderFilename);
-			}
-		}
-
 		SlowTask.EnterProgressFrame();
 
-		bool bGenerateProjectFiles = true;
-		{
-			bool bProjectHadCodeFiles = false;
-			{
-				TArray<FString> OutProjectCodeFilenames;
-				IFileManager::Get().FindFilesRecursive(OutProjectCodeFilenames, *FPaths::GameSourceDir(), TEXT("*.h"), true, false, false);
-				IFileManager::Get().FindFilesRecursive(OutProjectCodeFilenames, *FPaths::GameSourceDir(), TEXT("*.cpp"), true, false, false);
-				bProjectHadCodeFiles = OutProjectCodeFilenames.Num() > 0;
-			}
-
-			TArray<FString> CreatedFilesForExternalAppRead;
-			CreatedFilesForExternalAppRead.Reserve(CreatedFiles.Num());
-			for (const FString& CreatedFile : CreatedFiles)
-			{
-				CreatedFilesForExternalAppRead.Add(IFileManager::Get().ConvertToAbsolutePathForExternalAppForRead(*CreatedFile));
-			}
-
-			// First see if we can avoid a full generation by adding the new files to an already open project
-			if (bProjectHadCodeFiles && FSourceCodeNavigation::AddSourceFiles(CreatedFilesForExternalAppRead))
-			{
-				// We successfully added the new files to the solution, but we still need to run UBT with -gather to update any UBT makefiles
-				if (FDesktopPlatformModule::Get()->InvalidateMakefiles(FPaths::RootDir(), FPaths::GetProjectFilePath(), GWarn))
-				{
-					// We managed the gather, so we can skip running the full generate
-					bGenerateProjectFiles = false;
-				}
-			}
-		}
-
-		SlowTask.EnterProgressFrame();
-
-		bool bProjectFileUpdated = true;
-		if (bGenerateProjectFiles)
-		{
-			// Generate project files if we happen to be using a project file.
-			if (!FDesktopPlatformModule::Get()->GenerateProjectFiles(FPaths::RootDir(), FPaths::GetProjectFilePath(), GWarn))
-			{
-				ErrorString += LOCTEXT("FailedToGenerateProjectFiles", "Failed to generate project files.").ToString();
-				bProjectFileUpdated = false;
-			}
-		}
-
 		return ErrorString.IsEmpty();
 	}
 };
diff --git a/Engine/Source/Developer/BlueprintNativeCodeGen/Public/BlueprintNativeCodeGenModule.h b/Engine/Source/Developer/BlueprintNativeCodeGen/Public/BlueprintNativeCodeGenModule.h
--- a/Engine/Source/Developer/BlueprintNativeCodeGen/Public/BlueprintNativeCodeGenModule.h
+++ b/Engine/Source/Developer/BlueprintNativeCodeGen/Public/BlueprintNativeCodeGenModule.h
@@ -11,11 +11,17 @@
 class UBlueprint;
 enum class ESavePackageResult;
 
+struct FPlatformNativizationDetails
+{
+	FString PlatformName;
+	FString PlatformTargetDirectory;
+	FCompilerNativizationOptions CompilerNativizationOptions;
+};
+
 struct FNativeCodeGenInitData
 {
-	// This is an array of pairs, the pairs are PlatformName/PlatformTargetDirectory. These
-	// are determined by the cooker:
-	TArray< TPair< FString, FString > > CodegenTargets;
+	// This is an array platforms. These are determined by the cooker:
+	TArray< FPlatformNativizationDetails > CodegenTargets;
 
 	// Optional Manifest ManifestIdentifier, used for child cook processes that need a unique manifest name.
 	// The identifier is used to make a unique name for each platform that is converted.
@@ -36,7 +42,7 @@ public:
 	*
 	* Not for use with any kind of incremental cooking.
 	*/
-	FORCEINLINE static void InitializeModuleForRerunDebugOnly(const TArray< TPair< FString, FString > >& CodegenTargets);
+	FORCEINLINE static void InitializeModuleForRerunDebugOnly(const TArray<FPlatformNativizationDetails>& CodegenTargets);
 
 	/**
 	 * Wrapper function that retrieves the interface to this module from the 
@@ -77,7 +83,7 @@ public:
 	virtual const TMultiMap<FName, TAssetSubclassOf<UObject>>& GetFunctionsBoundToADelegate() = 0;
 protected:
 	virtual void Initialize(const FNativeCodeGenInitData& InitData) = 0;
-	virtual void InitializeForRerunDebugOnly(const TArray< TPair< FString, FString > >& CodegenTargets) = 0;
+	virtual void InitializeForRerunDebugOnly(const TArray<FPlatformNativizationDetails>& CodegenTargets) = 0;
 };
 
 void IBlueprintNativeCodeGenModule::InitializeModule(const FNativeCodeGenInitData& InitData)
@@ -86,7 +92,7 @@ void IBlueprintNativeCodeGenModule::InitializeModule(const FNativeCodeGenInitDat
 	Module.Initialize(InitData);
 }
 
-void IBlueprintNativeCodeGenModule::InitializeModuleForRerunDebugOnly(const TArray< TPair< FString, FString > >& CodegenTargets)
+void IBlueprintNativeCodeGenModule::InitializeModuleForRerunDebugOnly(const TArray<FPlatformNativizationDetails>& CodegenTargets)
 {
 	IBlueprintNativeCodeGenModule& Module = FModuleManager::LoadModuleChecked<IBlueprintNativeCodeGenModule>(GetModuleName());
 	Module.InitializeForRerunDebugOnly(CodegenTargets);
diff --git a/Engine/Source/Developer/BlueprintProfiler/Private/BlueprintProfiler.cpp b/Engine/Source/Developer/BlueprintProfiler/Private/BlueprintProfiler.cpp
--- a/Engine/Source/Developer/BlueprintProfiler/Private/BlueprintProfiler.cpp
+++ b/Engine/Source/Developer/BlueprintProfiler/Private/BlueprintProfiler.cpp
@@ -349,6 +349,15 @@ void FBlueprintProfiler::ProcessEventProfilingData()
 									SuspendedEventMap->Remove(LatentLinkId);
 								}
 							}
+							else 
+							{
+								const int32 EntryPointId = InstrumentationEventQueue[EventRangeToProcess.StartIdx].GetScriptCodeOffset();
+								UBlueprintGeneratedClass* BPGC = EventRangeToProcess.BlueprintContext->GetBlueprintClass().Get();
+								if (BPGC && BPGC->GetDebugData().IsValidEntryPoint(EntryPointId))
+								{
+									EventToProcess = (MakeShareable(new FScriptEventPlayback(EventRangeToProcess.BlueprintContext, EventRangeToProcess.InstanceName)));
+								}
+							}
 						}
 						else
 						{
@@ -376,8 +385,6 @@ void FBlueprintProfiler::ProcessEventProfilingData()
 			}
 		}
 	}
-	// We need to ensure this doesn't happen. but for now kill anything remaining in the queue.
-	InstrumentationEventQueue.SetNum(0);
 	// Update all active contexts if the display settings changed.
 	UBlueprintProfilerSettings* Settings = GetMutableDefault<UBlueprintProfilerSettings>();
 	if (Settings->GetPerformanceThresholdsModified())
diff --git a/Engine/Source/Developer/BlueprintProfiler/Private/ScriptInstrumentationPlayback.cpp b/Engine/Source/Developer/BlueprintProfiler/Private/ScriptInstrumentationPlayback.cpp
--- a/Engine/Source/Developer/BlueprintProfiler/Private/ScriptInstrumentationPlayback.cpp
+++ b/Engine/Source/Developer/BlueprintProfiler/Private/ScriptInstrumentationPlayback.cpp
@@ -819,7 +819,7 @@ void FBlueprintFunctionContext::CreateDelegatePinEvents()
 	};
 
 	FBlueprintDebugData& DebugData = BlueprintClass.Get()->GetDebugData();
-	const TMap<FName, FEdGraphPinReference>& PinEvents = DebugData.GetCompilerGeneratedEvents();
+	const TMap<int32, FName>& PinEvents = DebugData.GetEntryPoints();
 	if (PinEvents.Num())
 	{
 		TMap<const UEdGraphNode*, TArray<FPinDelegateDesc>> NodeEventDescs;
@@ -828,7 +828,7 @@ void FBlueprintFunctionContext::CreateDelegatePinEvents()
 		// Build event contexts per node
 		for (auto PinEvent : PinEvents)
 		{
-			if (UEdGraphPin* DelegatePin = PinEvent.Value.Get())
+			if (UEdGraphPin* DelegatePin = DebugData.FindSourcePinFromCodeLocation(BlueprintClass.Get()->UberGraphFunction, PinEvent.Key))
 			{
 				const UEdGraphNode* OwningNode = DelegatePin->GetOwningNode();
 				TArray<FPinDelegateDesc>& Events = NodeEventDescs.FindOrAdd(OwningNode);
@@ -836,14 +836,14 @@ void FBlueprintFunctionContext::CreateDelegatePinEvents()
 				FNodeExecutionContext& NodeContext = GetNodeExecutionContext(ScriptOffset);
 				if (NodeContext.IsValid())
 				{
-					const FName ScopedEventName = NodeContext.ProfilerContext->GetScopedEventName(PinEvent.Key);
+					const FName ScopedEventName = NodeContext.ProfilerContext->GetScopedEventName(PinEvent.Value);
 					Events.Add(FPinDelegateDesc(ScopedEventName, DelegatePin, ScriptOffset, NodeContext));
 				}
 				else
 				{
 					NodeContext.ProfilerContext = AsShared();
 					NodeContext.GraphNode = OwningNode;
-					const FName ScopedEventName = GetScopedEventName(PinEvent.Key);
+					const FName ScopedEventName = GetScopedEventName(PinEvent.Value);
 					Events.Add(FPinDelegateDesc(ScopedEventName, DelegatePin, ScriptOffset, NodeContext));
 				}
 			}
@@ -2953,21 +2953,7 @@ bool FScriptEventPlayback::Process(const TArray<FScriptInstrumentedEvent>& Signa
 						{
 							if (UWorld* WorldForEvent = InstanceObj->GetWorld())
 							{
-								FLatentActionManager& LatentActionManager = WorldForEvent->GetLatentActionManager();
-								TSet<int32> UUIDSet;
-								int32 UUID = INDEX_NONE;
-								LatentActionManager.GetActiveUUIDs(InstanceObj, UUIDSet);
-								for (auto SetEntry : UUIDSet)
-								{
-									UUID = SetEntry;
-								}
-								if (UUID != INDEX_NONE)
-								{
-									if (FDelayAction* Action = LatentActionManager.FindExistingAction<FDelayAction>(InstanceObj, UUID))
-									{
-										LatentLinkId = Action->OutputLink;
-									}
-								}
+								LatentLinkId = CurrSignal.GetScriptCodeOffset();
 							}
 						}
 						break;
@@ -3116,7 +3102,7 @@ bool FScriptEventPlayback::Process(const TArray<FScriptInstrumentedEvent>& Signa
 						{
 							PerfData = PureChainNode->GetOrAddPerfDataByInstanceAndTracePath(InstanceName, NodeTracePath);
 							PerfData->AddEventTiming(PureChainDuration);
-						}
+  						}
 						PureChainEntryTime = NodeEntryTime = 0.0;
 						break;
 					}
diff --git a/Engine/Source/Editor/BlueprintGraph/Classes/K2Node.h b/Engine/Source/Editor/BlueprintGraph/Classes/K2Node.h
--- a/Engine/Source/Editor/BlueprintGraph/Classes/K2Node.h
+++ b/Engine/Source/Editor/BlueprintGraph/Classes/K2Node.h
@@ -256,9 +256,6 @@ class UK2Node : public UEdGraphNode
 	 */
 	virtual bool IsNodeSafeToIgnore() const { return false; }
 
-	/** Tries to get a template object from this node. Will only work for 'Add Component' nodes */
-	virtual UActorComponent* GetTemplateFromNode() const { return NULL; }
-
 	/** Called at the end of ReconstructNode, allows node specific work to be performed */
 	BLUEPRINTGRAPH_API virtual void PostReconstructNode();
 
diff --git a/Engine/Source/Editor/BlueprintGraph/Classes/K2Node_AddComponent.h b/Engine/Source/Editor/BlueprintGraph/Classes/K2Node_AddComponent.h
--- a/Engine/Source/Editor/BlueprintGraph/Classes/K2Node_AddComponent.h
+++ b/Engine/Source/Editor/BlueprintGraph/Classes/K2Node_AddComponent.h
@@ -26,6 +26,9 @@ class UK2Node_AddComponent : public UK2Node_CallFunction
 	UPROPERTY()
 	FString TemplateBlueprint;
 
+	UPROPERTY()
+	UClass* TemplateType;
+
 	//~ Begin UObject Interface
 	virtual void Serialize(FArchive& Ar) override;
 	virtual void PostDuplicate(bool bDuplicateForPIE) override;
@@ -40,10 +43,11 @@ class UK2Node_AddComponent : public UK2Node_CallFunction
 	virtual FString GetDocumentationLink() const override;
 	virtual FString GetDocumentationExcerptName() const override;
 	virtual bool IsCompatibleWithGraph(UEdGraph const* Graph) const override;
+	virtual void ReconstructNode() override;
 	//~ End UEdGraphNode Interface
 
 	//~ Begin K2Node Interface
-	virtual UActorComponent* GetTemplateFromNode() const override;
+	virtual void PostReconstructNode() override;
 	virtual void ValidateNodeDuringCompilation(class FCompilerResultsLog& MessageLog) const override;
 	virtual void ReallocatePinsDuringReconstruction(TArray<UEdGraphPin*>& OldPins) override;
 	virtual void ExpandNode(class FKismetCompilerContext& CompilerContext, UEdGraph* SourceGraph) override;
@@ -69,6 +73,9 @@ class UK2Node_AddComponent : public UK2Node_CallFunction
 		return FindPinChecked(TEXT("bManualAttachment"));
 	}
 
+	/** Tries to get a template object from this node. */
+	BLUEPRINTGRAPH_API UActorComponent* GetTemplateFromNode() const;
+
 	/** Helper method used to generate a new, unique component template name. */
 	FName MakeNewComponentTemplateName(UObject* InOuter, UClass* InComponentClass);
 
@@ -84,7 +91,7 @@ private:
 		return FindPin(TEXT("TemplateName"));
 	}
 
-	const UClass* GetSpawnedType() const;
+	UClass* GetSpawnedType() const;	
 };
 
 
diff --git a/Engine/Source/Editor/BlueprintGraph/Classes/K2Node_CallFunction.h b/Engine/Source/Editor/BlueprintGraph/Classes/K2Node_CallFunction.h
--- a/Engine/Source/Editor/BlueprintGraph/Classes/K2Node_CallFunction.h
+++ b/Engine/Source/Editor/BlueprintGraph/Classes/K2Node_CallFunction.h
@@ -74,7 +74,6 @@ public:
 	// UEdGraphNode interface
 	virtual void GetPinHoverText(const UEdGraphPin& Pin, FString& HoverTextOut) const override;
 	virtual void AllocateDefaultPins() override;
-	virtual void DestroyNode() override;
 	virtual FLinearColor GetNodeTitleColor() const override;
 	virtual FText GetTooltipText() const override;
 	virtual FText GetNodeTitle(ENodeTitleType::Type TitleType) const override;
diff --git a/Engine/Source/Editor/BlueprintGraph/Classes/K2Node_EaseFunction.h b/Engine/Source/Editor/BlueprintGraph/Classes/K2Node_EaseFunction.h
--- a/Engine/Source/Editor/BlueprintGraph/Classes/K2Node_EaseFunction.h
+++ b/Engine/Source/Editor/BlueprintGraph/Classes/K2Node_EaseFunction.h
@@ -55,6 +55,7 @@ class UK2Node_EaseFunction : public UK2Node
 	virtual bool IsConnectionDisallowed(const UEdGraphPin* MyPin, const UEdGraphPin* OtherPin, FString& OutReason) const override;
 	virtual void GetMenuActions(FBlueprintActionDatabaseRegistrar& ActionRegistrar) const override;
 	virtual FText GetMenuCategory() const override;
+	virtual bool IsNodePure() const override { return true; }
 	//~ End UK2Node Interface.
 
 private:
diff --git a/Engine/Source/Editor/BlueprintGraph/Classes/K2Node_Knot.h b/Engine/Source/Editor/BlueprintGraph/Classes/K2Node_Knot.h
--- a/Engine/Source/Editor/BlueprintGraph/Classes/K2Node_Knot.h
+++ b/Engine/Source/Editor/BlueprintGraph/Classes/K2Node_Knot.h
@@ -51,4 +51,9 @@ public:
 	{
 		return Pins[1];
 	}
+
+private:
+	void PropagatePinType();
+	void PropagatePinTypeFromInput();
+	void PropagatePinTypeFromOutput();
 };
