Fix monolithic editor compile errors.

* IMPLEMENT_MODULE must take the module name as the second parameter.
* Some functions declared in a header but not marked as inline, causing duplicate symbol errors when linking.
* GetPathPostfix() function duplicated in two files.
* lws_debugLog function duplicated in two files.

diff --git a/Engine/Plugins/Experimental/HTML5Networking/Source/HTML5Networking/Private/WebSocketServer.cpp b/Engine/Plugins/Experimental/HTML5Networking/Source/HTML5Networking/Private/WebSocketServer.cpp
--- a/Engine/Plugins/Experimental/HTML5Networking/Source/HTML5Networking/Private/WebSocketServer.cpp
+++ b/Engine/Plugins/Experimental/HTML5Networking/Source/HTML5Networking/Private/WebSocketServer.cpp
@@ -25,7 +25,7 @@ struct PerSessionDataServer
 static int unreal_networking_server(struct lws *wsi, enum lws_callback_reasons reason, void *user, void *in, size_t len);
 
 #if !UE_BUILD_SHIPPING
-	void lws_debugLog(int level, const char *line)
+	inline void lws_debugLog(int level, const char *line)
 	{
 		UE_LOG(LogHTML5Networking, Log, TEXT("websocket server: %s"), ANSI_TO_TCHAR(line));
 	}
diff --git a/Engine/Source/Developer/BlueprintCompilerCppBackend/Private/BlueprintCompilerCppBackendUtils.cpp b/Engine/Source/Developer/BlueprintCompilerCppBackend/Private/BlueprintCompilerCppBackendUtils.cpp
--- a/Engine/Source/Developer/BlueprintCompilerCppBackend/Private/BlueprintCompilerCppBackendUtils.cpp
+++ b/Engine/Source/Developer/BlueprintCompilerCppBackend/Private/BlueprintCompilerCppBackendUtils.cpp
@@ -19,18 +19,7 @@
 #include "BlueprintCompilerCppBackend.h"
 #include "Animation/AnimBlueprint.h"
 
-FString GetPathPostfix(const UObject* ForObject)
-{
-	FString FullAssetName = ForObject->GetOutermost()->GetPathName();
-	if (FullAssetName.StartsWith(TEXT("/Temp/__TEMP_BP__"), ESearchCase::CaseSensitive))
-	{
-		FullAssetName.RemoveFromStart(TEXT("/Temp/__TEMP_BP__"), ESearchCase::CaseSensitive);
-	}
-	FString AssetName = FPackageName::GetLongPackageAssetName(FullAssetName);
-	// append a hash of the path, this uniquely identifies assets with the same name, but different folders:
-	FullAssetName.RemoveFromEnd(AssetName);
-	return FString::Printf(TEXT("%u"), FCrc::MemCrc32(*FullAssetName, FullAssetName.Len()*sizeof(TCHAR)));
-}
+extern ENGINE_API FString GetPathPostfix(const UObject* ForObject);
 
 FString FEmitterLocalContext::GenerateUniqueLocalName()
 {
diff --git a/Engine/Source/Runtime/Engine/Private/UserDefinedStruct.cpp b/Engine/Source/Runtime/Engine/Private/UserDefinedStruct.cpp
--- a/Engine/Source/Runtime/Engine/Private/UserDefinedStruct.cpp
+++ b/Engine/Source/Runtime/Engine/Private/UserDefinedStruct.cpp
@@ -257,7 +257,7 @@ FGuid UUserDefinedStruct::GetCustomGuid() const
 	return Guid;
 }
 
-FString GetPathPostfix(const UObject* ForObject)
+ENGINE_API FString GetPathPostfix(const UObject* ForObject)
 {
 	FString FullAssetName = ForObject->GetOutermost()->GetPathName();
 	if (FullAssetName.StartsWith(TEXT("/Temp/__TEMP_BP__"), ESearchCase::CaseSensitive))
diff --git a/Engine/Source/Runtime/NetworkFileSystem/Private/NetworkFileServerHttp.cpp b/Engine/Source/Runtime/NetworkFileSystem/Private/NetworkFileServerHttp.cpp
--- a/Engine/Source/Runtime/NetworkFileSystem/Private/NetworkFileServerHttp.cpp
+++ b/Engine/Source/Runtime/NetworkFileSystem/Private/NetworkFileServerHttp.cpp
@@ -121,7 +121,7 @@ bool FNetworkFileServerHttp::IsItReadyToAcceptConnections(void) const
 }
 
 #if UE_BUILD_DEBUG
-void lws_debugLog(int level, const char *line)
+inline void lws_debugLog(int level, const char *line)
 {
 	UE_LOG(LogFileServer, Warning, TEXT(" LibWebsocket: %s"), ANSI_TO_TCHAR(line));
 }
