Safer fix for static analysis issue

diff --git a/Engine/Source/Developer/BlueprintCompilerCppBackend/Private/BlueprintCompilerCppBackendUtils.cpp b/Engine/Source/Developer/BlueprintCompilerCppBackend/Private/BlueprintCompilerCppBackendUtils.cpp
--- a/Engine/Source/Developer/BlueprintCompilerCppBackend/Private/BlueprintCompilerCppBackendUtils.cpp
+++ b/Engine/Source/Developer/BlueprintCompilerCppBackend/Private/BlueprintCompilerCppBackendUtils.cpp
@@ -260,13 +260,13 @@ FString FEmitterLocalContext::FindGloballyMappedObject(const UObject* Object, co
 		}
 	}
 
-	if (ActualUserStruct)
+	if (Object && ActualUserStruct)
 	{
 		// For user structs, the default action of loading is unsafe so call the wrapper function
 		return FString::Printf(TEXT("CastChecked<%s>(FConvertedBlueprintsDependencies::LoadObjectForStructConstructor(%s::StaticStruct(),TEXT(\"%s\")), ECastCheckedType::NullAllowed)")
 			, *ClassString()
 			, *FEmitHelper::GetCppName(ActualUserStruct)
-			, *GetPathNameSafe(Object).ReplaceCharWithEscapedChar());			
+			, *(Object->GetPathName().ReplaceCharWithEscapedChar()));
 	}
 
 	return FString{};
