Fix UHT warning about duplicating UFUNCTION() meta on nativized child Blueprint classes that override at least one BPIE method from the parent class.

- Mirrored from //UE4/Dev-Framework (4111613).

#rnx

diff --git a/Engine/Source/Developer/BlueprintCompilerCppBackend/Private/BlueprintCompilerCppBackendBase.cpp b/Engine/Source/Developer/BlueprintCompilerCppBackend/Private/BlueprintCompilerCppBackendBase.cpp
--- a/Engine/Source/Developer/BlueprintCompilerCppBackend/Private/BlueprintCompilerCppBackendBase.cpp
+++ b/Engine/Source/Developer/BlueprintCompilerCppBackend/Private/BlueprintCompilerCppBackendBase.cpp
@@ -680,6 +680,7 @@ TArray<FString> FBlueprintCompilerCppBackendBase::ConstructFunctionDeclaration(F
 		UFunction* const OriginalFunction = FEmitHelper::GetOriginalFunction(Function);
 		TArray<FString> AdditionalMetaData;
 		TArray<FString> AdditionalTags;
+		bool bIsBIEOverride = false;
 		bool bGenerateAsNativeEventImplementation = false;
 		const bool bNetImplementation = !bInInterface && Function->HasAllFunctionFlags(FUNC_Net) && !Function->HasAnyFunctionFlags(FUNC_NetResponse);
 
@@ -702,9 +703,38 @@ TArray<FString> FBlueprintCompilerCppBackendBase::ConstructFunctionDeclaration(F
 		}
 		else if (FEmitHelper::ShouldHandleAsImplementableEvent(Function) || bBPInterfaceImplementation)
 		{
-			//The function "bpf__BIE__pf" should never be called directly. Only via function "BIE" with generated implementation.
+			// The function "bpf__BIE__pf" should never be called directly. Only via function "BIE" with generated implementation.
 			bIsVirtual = false;
 			AdditionalMetaData.Emplace(TEXT("CppFromBpEvent"));
+			
+			// Get the owner of the current function context.
+			const UClass* OwnerClass = Function->GetOwnerClass();
+			check(OwnerClass);
+
+			// Get the owner's parent class context.
+			const UClass* SuperClass = OwnerClass->GetSuperClass();
+			check(SuperClass);
+
+			// Only need to check for an override if the parent class is non-native, as it may also be convertible in that case.
+			if (!SuperClass->HasAllClassFlags(CLASS_Native))
+			{
+				// See if the same function is implemented there. Note: Cannot use 'OriginalFunction' here because that goes all the way back to the earliest antecedent, which is typically
+				// a native engine class. In here, we're only concerned with the closest non-native parent class that is both being converted and also implements the same function as a BPIE.
+				const FName FunctionName = Function->GetFName();
+				if (const UFunction* SuperFunction = SuperClass->FindFunctionByName(FunctionName))
+				{
+					// Ensure that we reference the inherited class that declares the function. This may be a bit farther up in the hierarchy than our immediate parent class, after the search.
+					SuperClass = SuperFunction->GetOwnerClass();
+
+					// We are technically overriding the BPIE in a converted child class if the parent class already implements it. However, it is not a virtual function, so we can't treat it
+					// as a "normal" override (i.e. - we don't use either the 'virtual' or the 'override' keywords when declaring it in the header).
+					if (const UBlueprintGeneratedClass* ParentBPGC = Cast<UBlueprintGeneratedClass>(SuperClass))
+					{
+						// In that case, we need to skip the UFUNCTION() markup if the parent class is also being converted, to avoid a UHT complaint about a redefinition of UFUNCTION() meta.
+						bIsBIEOverride = EmitterContext.Dependencies.WillClassBeConverted(ParentBPGC);
+					}
+				}
+			}
 		}
 
 		ensure(!bIsVirtual || Function->IsSignatureCompatibleWith(OriginalFunction));
@@ -756,7 +786,7 @@ TArray<FString> FBlueprintCompilerCppBackendBase::ConstructFunctionDeclaration(F
 		{
 			FunctionHeaderName = FunctionBodyName;
 		}
-		else if (!bGenerateAsNativeEventImplementation && !bSkipMacro)
+		else if (!bIsBIEOverride && !bGenerateAsNativeEventImplementation && !bSkipMacro)
 		{
 			MacroUFUNCTION = FEmitHelper::EmitUFuntion(Function, AdditionalTags, AdditionalMetaData);
 		}
