Fix invalid child actor template cast in nativized code when the child actor class type is a converted Blueprint class asset.

- Mirrored from //UE4/Dev-Framework (4111638).

diff --git a/Engine/Source/Developer/BlueprintCompilerCppBackend/Private/BlueprintCompilerCppBackendValueHelper.cpp b/Engine/Source/Developer/BlueprintCompilerCppBackend/Private/BlueprintCompilerCppBackendValueHelper.cpp
--- a/Engine/Source/Developer/BlueprintCompilerCppBackend/Private/BlueprintCompilerCppBackendValueHelper.cpp
+++ b/Engine/Source/Developer/BlueprintCompilerCppBackend/Private/BlueprintCompilerCppBackendValueHelper.cpp
@@ -109,7 +109,7 @@ void FEmitDefaultValueHelper::OuterGenerate(FEmitterLocalContext& Context
 					if (ObjectPropertyValue && ObjectPropertyValue->IsDefaultSubobject())
 					{
 						UClass* SubobjectClass = ObjectPropertyValue->GetClass();
-						OverrideTypeDeclaration = FString::Printf(TEXT("%s%s*"), SubobjectClass->GetPrefixCPP(), *SubobjectClass->GetName());
+						OverrideTypeDeclaration = FString::Printf(TEXT("%s*"), *FEmitHelper::GetCppName(SubobjectClass));
 					}
 				}
 
