Fix a regression that causes some UMG widgets to no longer animate in a nativized cooked build.

- Mirrored from //UE4/Dev-Framework (4121263).

#rnx

diff --git a/Engine/Source/Developer/BlueprintCompilerCppBackend/Private/BlueprintCompilerCppBackendUMG.cpp b/Engine/Source/Developer/BlueprintCompilerCppBackend/Private/BlueprintCompilerCppBackendUMG.cpp
--- a/Engine/Source/Developer/BlueprintCompilerCppBackend/Private/BlueprintCompilerCppBackendUMG.cpp
+++ b/Engine/Source/Developer/BlueprintCompilerCppBackend/Private/BlueprintCompilerCppBackendUMG.cpp
@@ -126,6 +126,22 @@ bool FBackendHelperUMG::SpecialStructureConstructorUMG(const UStruct* Struct, co
 {
 	check(ValuePtr || !OutResult);
 
+	auto FrameNumberRangeBoundConstructorLambda = [](const TRangeBound<FFrameNumber>& InRangeBound, const FFrameNumber& InRangeBoundValue) -> FString
+	{
+		if (InRangeBound.IsExclusive())
+		{
+			return FString::Printf(TEXT("TRangeBound<FFrameNumber>::Exclusive(%d)"), InRangeBoundValue.Value);
+		}
+		else if (InRangeBound.IsInclusive())
+		{
+			return FString::Printf(TEXT("TRangeBound<FFrameNumber>::Inclusive(%d)"), InRangeBoundValue.Value);
+		}
+		else
+		{
+			return FString::Printf(TEXT("TRangeBound<FFrameNumber>::Open()"));
+		}
+	};
+
 	if (FSectionEvaluationData::StaticStruct() == Struct)
 	{
 		if (OutResult)
@@ -151,22 +167,6 @@ bool FBackendHelperUMG::SpecialStructureConstructorUMG(const UStruct* Struct, co
 	{
 		if (OutResult)
 		{
-			auto MovieSceneSegmentRangeBoundConstructorLambda = [](const TRangeBound<FFrameNumber>& InRangeBound, const FFrameNumber& InRangeBoundValue) -> FString
-			{
-				if (InRangeBound.IsExclusive())
-				{
-					return FString::Printf(TEXT("TRangeBound<FFrameNumber>::Exclusive(%d)"), InRangeBoundValue.Value);
-				}
-				else if (InRangeBound.IsInclusive())
-				{
-					return FString::Printf(TEXT("TRangeBound<FFrameNumber>::Inclusive(%d)"), InRangeBoundValue.Value);
-				}
-				else
-				{
-					return FString::Printf(TEXT("TRangeBound<FFrameNumber>::Open()"));
-				}
-			};
-
 			const FMovieSceneSegment* MovieSceneSegment = reinterpret_cast<const FMovieSceneSegment*>(ValuePtr);
 			FString SegmentsInitializerList;
 			for (const FSectionEvaluationData& SectionEvaluationData : MovieSceneSegment->Impls)
@@ -181,14 +181,29 @@ bool FBackendHelperUMG::SpecialStructureConstructorUMG(const UStruct* Struct, co
 					, &SectionEvaluationDataStr);
 				SegmentsInitializerList += SectionEvaluationDataStr;
 			}
-			const FString LowerBoundStr = MovieSceneSegmentRangeBoundConstructorLambda(MovieSceneSegment->Range.GetLowerBound()
-				, MovieSceneSegment->Range.GetLowerBound().IsClosed() ? MovieSceneSegment->Range.GetLowerBoundValue() : FFrameNumber());
-			const FString UpperBoundStr = MovieSceneSegmentRangeBoundConstructorLambda(MovieSceneSegment->Range.GetUpperBound()
-				, MovieSceneSegment->Range.GetUpperBound().IsClosed() ? MovieSceneSegment->Range.GetUpperBoundValue() : FFrameNumber());
+			const FString LowerBoundStr = FrameNumberRangeBoundConstructorLambda(MovieSceneSegment->Range.GetLowerBound(),
+				MovieSceneSegment->Range.GetLowerBound().IsClosed() ? MovieSceneSegment->Range.GetLowerBoundValue() : FFrameNumber());
+			const FString UpperBoundStr = FrameNumberRangeBoundConstructorLambda(MovieSceneSegment->Range.GetUpperBound(),
+				MovieSceneSegment->Range.GetUpperBound().IsClosed() ? MovieSceneSegment->Range.GetUpperBoundValue() : FFrameNumber());
 			*OutResult = FString::Printf(TEXT("FMovieSceneSegment(TRange<FFrameNumber>(%s, %s), {%s})"), *LowerBoundStr, *UpperBoundStr, *SegmentsInitializerList);
 		}
 		return true;
 	}
+
+	if (FMovieSceneFrameRange::StaticStruct() == Struct)
+	{
+		if (OutResult)
+		{
+			const FMovieSceneFrameRange* MovieSceneFrameRange = reinterpret_cast<const FMovieSceneFrameRange*>(ValuePtr);
+			const FString LowerBoundStr = FrameNumberRangeBoundConstructorLambda(MovieSceneFrameRange->Value.GetLowerBound(),
+				MovieSceneFrameRange->Value.GetLowerBound().IsClosed() ? MovieSceneFrameRange->Value.GetLowerBoundValue() : FFrameNumber());
+			const FString UpperBoundStr = FrameNumberRangeBoundConstructorLambda(MovieSceneFrameRange->Value.GetUpperBound(),
+				MovieSceneFrameRange->Value.GetUpperBound().IsClosed() ? MovieSceneFrameRange->Value.GetUpperBoundValue() : FFrameNumber());
+			*OutResult = FString::Printf(TEXT("FMovieSceneFrameRange(TRange<FFrameNumber>(%s, %s))"), *LowerBoundStr, *UpperBoundStr);
+		}
+		return true;
+	}
+
 	return false;
 }
 
diff --git a/Engine/Source/Runtime/UMG/Private/UserWidget.cpp b/Engine/Source/Runtime/UMG/Private/UserWidget.cpp
--- a/Engine/Source/Runtime/UMG/Private/UserWidget.cpp
+++ b/Engine/Source/Runtime/UMG/Private/UserWidget.cpp
@@ -1578,8 +1578,9 @@ void UUserWidget::UpdateCanTick()
 		bool bCanTick = false;
 		if (TickFrequency == EWidgetTickFrequency::Auto)
 		{
+			// Note: WidgetBPClass can be NULL in a cooked build, if the Blueprint has been nativized (in that case, it will be a UDynamicClass type).
 			UWidgetBlueprintGeneratedClass* WidgetBPClass = Cast<UWidgetBlueprintGeneratedClass>(GetClass());
-			bCanTick |= WidgetBPClass->ClassRequiresNativeTick();
+			bCanTick |= !WidgetBPClass || WidgetBPClass->ClassRequiresNativeTick();
 			bCanTick |= bHasScriptImplementedTick;
 			bCanTick |= World->GetLatentActionManager().GetNumActionsForObject(this) != 0;
 			bCanTick |= ActiveSequencePlayers.Num() > 0;
