Fix UHT warning about duplicating UFUNCTION() meta on nativized child Blueprint classes that override at least one non-native method defined in a parent Blueprint class.

- Mirrored from //UE4/Dev-Framework (4128271).

#rnx

diff --git a/Engine/Source/Developer/BlueprintCompilerCppBackend/Private/BlueprintCompilerCppBackendBase.cpp b/Engine/Source/Developer/BlueprintCompilerCppBackend/Private/BlueprintCompilerCppBackendBase.cpp
--- a/Engine/Source/Developer/BlueprintCompilerCppBackend/Private/BlueprintCompilerCppBackendBase.cpp
+++ b/Engine/Source/Developer/BlueprintCompilerCppBackend/Private/BlueprintCompilerCppBackendBase.cpp
@@ -680,7 +680,7 @@ TArray<FString> FBlueprintCompilerCppBackendBase::ConstructFunctionDeclaration(F
 		UFunction* const OriginalFunction = FEmitHelper::GetOriginalFunction(Function);
 		TArray<FString> AdditionalMetaData;
 		TArray<FString> AdditionalTags;
-		bool bIsBIEOverride = false;
+		bool bGenerateAsNonNativeOverride = false;
 		bool bGenerateAsNativeEventImplementation = false;
 		const bool bNetImplementation = !bInInterface && Function->HasAllFunctionFlags(FUNC_Net) && !Function->HasAnyFunctionFlags(FUNC_NetResponse);
 
@@ -731,11 +731,19 @@ TArray<FString> FBlueprintCompilerCppBackendBase::ConstructFunctionDeclaration(F
 					if (const UBlueprintGeneratedClass* ParentBPGC = Cast<UBlueprintGeneratedClass>(SuperClass))
 					{
 						// In that case, we need to skip the UFUNCTION() markup if the parent class is also being converted, to avoid a UHT complaint about a redefinition of UFUNCTION() meta.
-						bIsBIEOverride = EmitterContext.Dependencies.WillClassBeConverted(ParentBPGC);
+						bGenerateAsNonNativeOverride = EmitterContext.Dependencies.WillClassBeConverted(ParentBPGC);
 					}
 				}
 			}
 		}
+		else if(OriginalFuncOwnerAsBPGC && Function != OriginalFunction)
+		{
+			// We assume parent classes will always be converted along with the child class.
+			check(EmitterContext.Dependencies.WillClassBeConverted(OriginalFuncOwnerAsBPGC));
+
+			// Skip the UFUNCTION() markup when the original class is also being converted, to avoid a UHT complaint about a redefinition of UFUNCTION() meta.
+			bGenerateAsNonNativeOverride = true;
+		}
 
 		ensure(!bIsVirtual || Function->IsSignatureCompatibleWith(OriginalFunction));
 		bIsOverride = bGenerateAsNativeEventImplementation || (bIsVirtual && (Function != OriginalFunction));
@@ -786,7 +794,7 @@ TArray<FString> FBlueprintCompilerCppBackendBase::ConstructFunctionDeclaration(F
 		{
 			FunctionHeaderName = FunctionBodyName;
 		}
-		else if (!bIsBIEOverride && !bGenerateAsNativeEventImplementation && !bSkipMacro)
+		else if (!bGenerateAsNonNativeOverride && !bGenerateAsNativeEventImplementation && !bSkipMacro)
 		{
 			MacroUFUNCTION = FEmitHelper::EmitUFuntion(Function, AdditionalTags, AdditionalMetaData);
 		}
