Disabled redundent BP nativization tests (https://udn.unrealengine.com/questions/436617)

diff --git a/Engine/Source/Developer/BlueprintCompilerCppBackend/Private/BPCompilerTests.cpp b/Engine/Source/Developer/BlueprintCompilerCppBackend/Private/BPCompilerTests.cpp
--- a/Engine/Source/Developer/BlueprintCompilerCppBackend/Private/BPCompilerTests.cpp
+++ b/Engine/Source/Developer/BlueprintCompilerCppBackend/Private/BPCompilerTests.cpp
@@ -200,8 +200,8 @@ static void Call(UObject* Target, const TCHAR* FunctionName, void* Args = nullpt
 	}
 }
 
-// Remove EAutomationTestFlags::Disabled to enable these tests, note that these will need to be moved into the ClientContext because we can only test cooked content:
-static const uint32 CompilerTestFlags = EAutomationTestFlags::ClientContext | EAutomationTestFlags::EngineFilter;
+// Remove EAutomationTestFlags::Disabled to enable these tests, note that these will need to be moved into the ClientContext because we can only test nativized cooked content:
+static const uint32 CompilerTestFlags = EAutomationTestFlags::ClientContext | EAutomationTestFlags::EngineFilter | EAutomationTestFlags::Disabled;
 
 // Tests:
 IMPLEMENT_SIMPLE_AUTOMATION_TEST(FBPCompilerArrayTest, "Project.Blueprints.NativeBackend.ArrayTest", CompilerTestFlags)
