Fix failure to build nativized Match3 (regression).

Change summary:
- Modified FEmitHelper::MetaDataCanBeNative() to exclude 'ExpandEnumAsExecs' metadata from UFUNCTION() markup in converted code.

#rnx

diff --git a/Engine/Source/Developer/BlueprintCompilerCppBackend/Private/BlueprintCompilerCppBackendUtils.cpp b/Engine/Source/Developer/BlueprintCompilerCppBackend/Private/BlueprintCompilerCppBackendUtils.cpp
--- a/Engine/Source/Developer/BlueprintCompilerCppBackend/Private/BlueprintCompilerCppBackendUtils.cpp
+++ b/Engine/Source/Developer/BlueprintCompilerCppBackend/Private/BlueprintCompilerCppBackendUtils.cpp
@@ -585,6 +585,10 @@ bool FEmitHelper::MetaDataCanBeNative(const FName MetaDataName, const UField* Fi
 	{
 		return false;
 	}
+	if (MetaDataName == TEXT("ExpandEnumAsExecs"))	// applicable to editor only
+	{
+		return false;
+	}
 	if (const UFunction* Function = Cast<const UFunction>(Field))
 	{
 		const UProperty* Param = Function->FindPropertyByName(MetaDataName);
