Fix broken event function overrides in nativized Blueprint child classes.

Change summary:
- Modified FBlueprintCompilerCppBackendBase::ConstructFunctionDeclaration() to declare BPIEs as virtual and BPIE overrides as virtual/override. This allows for indirection to the right function pointer through the vtable, since we no longer have explicit UHT boilerplate code being generated for function overrides due to removal of UFUNCTION markup, which was done to comply with UHT.

#rnx

diff --git a/Engine/Source/Developer/BlueprintCompilerCppBackend/Private/BlueprintCompilerCppBackendBase.cpp b/Engine/Source/Developer/BlueprintCompilerCppBackend/Private/BlueprintCompilerCppBackendBase.cpp
--- a/Engine/Source/Developer/BlueprintCompilerCppBackend/Private/BlueprintCompilerCppBackendBase.cpp
+++ b/Engine/Source/Developer/BlueprintCompilerCppBackend/Private/BlueprintCompilerCppBackendBase.cpp
@@ -682,6 +682,8 @@ TArray<FString> FBlueprintCompilerCppBackendBase::ConstructFunctionDeclaration(F
 		TArray<FString> AdditionalTags;
 		bool bGenerateAsNonNativeOverride = false;
 		bool bGenerateAsNativeEventImplementation = false;
+		const bool bShouldHandleAsNativeEvent = FEmitHelper::ShouldHandleAsNativeEvent(Function);
+		const bool bShouldHandleAsNonNativeEvent = FEmitHelper::ShouldHandleAsImplementableEvent(Function);
 		const bool bNetImplementation = !bInInterface && Function->HasAllFunctionFlags(FUNC_Net) && !Function->HasAnyFunctionFlags(FUNC_NetResponse);
 
 		const UBlueprintGeneratedClass* const OriginalFuncOwnerAsBPGC = Cast<UBlueprintGeneratedClass>(OriginalFunction->GetOwnerClass());
@@ -695,16 +697,15 @@ TArray<FString> FBlueprintCompilerCppBackendBase::ConstructFunctionDeclaration(F
 		{
 			FunctionBodyName = FunctionHeaderName + TEXT("_Implementation");
 		}
-		else if (FEmitHelper::ShouldHandleAsNativeEvent(Function))
+		else if (bShouldHandleAsNativeEvent)
 		{
 			bGenerateAsNativeEventImplementation = true;
 			FunctionBodyName = FunctionHeaderName = FEmitHelper::GetCppName(OriginalFunction) + TEXT("_Implementation");
 			bAddConst = OriginalFunction->HasAllFunctionFlags(FUNC_Const);
 		}
-		else if (FEmitHelper::ShouldHandleAsImplementableEvent(Function) || bBPInterfaceImplementation)
+		else if (bShouldHandleAsNonNativeEvent || bBPInterfaceImplementation)
 		{
 			// The function "bpf__BIE__pf" should never be called directly. Only via function "BIE" with generated implementation.
-			bIsVirtual = false;
 			AdditionalMetaData.Emplace(TEXT("CppFromBpEvent"));
 			
 			// Get the owner of the current function context.
@@ -726,8 +727,7 @@ TArray<FString> FBlueprintCompilerCppBackendBase::ConstructFunctionDeclaration(F
 					// Ensure that we reference the inherited class that declares the function. This may be a bit farther up in the hierarchy than our immediate parent class, after the search.
 					SuperClass = SuperFunction->GetOwnerClass();
 
-					// We are technically overriding the BPIE in a converted child class if the parent class already implements it. However, it is not a virtual function, so we can't treat it
-					// as a "normal" override (i.e. - we don't use either the 'virtual' or the 'override' keywords when declaring it in the header).
+					// We are overriding the BPIE in a converted child class if the parent class is also being converted and already implements it.
 					if (const UBlueprintGeneratedClass* ParentBPGC = Cast<UBlueprintGeneratedClass>(SuperClass))
 					{
 						// In that case, we need to skip the UFUNCTION() markup if the parent class is also being converted, to avoid a UHT complaint about a redefinition of UFUNCTION() meta.
@@ -746,7 +746,7 @@ TArray<FString> FBlueprintCompilerCppBackendBase::ConstructFunctionDeclaration(F
 		}
 
 		ensure(!bIsVirtual || Function->IsSignatureCompatibleWith(OriginalFunction));
-		bIsOverride = bGenerateAsNativeEventImplementation || (bIsVirtual && (Function != OriginalFunction));
+		bIsOverride = bGenerateAsNativeEventImplementation || bGenerateAsNonNativeOverride || (bIsVirtual && !bShouldHandleAsNonNativeEvent && !bBPInterfaceImplementation && (Function != OriginalFunction));
 
 		auto PreliminaryConditionsToSkipMacroUFUNC = [](UFunction* InFunction) -> bool 
 		{
