Fix assertion in a nativized build on attempt to register an instanced subobject that was constructed based on an archetype that was originally instanced through the editor UI.

Change summary:
- Modified FEmitDefaultValueHelper::HandleInstancedSubobject() to remove the RF_ArchetypeObject flag from the flags parameter in the NewObject() call that's emitted to the converted class ctor.

diff --git a/Engine/Source/Developer/BlueprintCompilerCppBackend/Private/BlueprintCompilerCppBackendValueHelper.cpp b/Engine/Source/Developer/BlueprintCompilerCppBackend/Private/BlueprintCompilerCppBackendValueHelper.cpp
--- a/Engine/Source/Developer/BlueprintCompilerCppBackend/Private/BlueprintCompilerCppBackendValueHelper.cpp
+++ b/Engine/Source/Developer/BlueprintCompilerCppBackend/Private/BlueprintCompilerCppBackendValueHelper.cpp
@@ -2079,8 +2079,9 @@ FString FEmitDefaultValueHelper::HandleInstancedSubobject(FEmitterLocalContext&
 			}
 			else
 			{
+				const int32 ObjectFlags = ((int32)Object->GetFlags() & ~RF_ArchetypeObject);
 				Context.AddLine(FString::Printf(TEXT("auto %s = NewObject<%s>(%s, TEXT(\"%s\"), (EObjectFlags)0x%08x);")
-					, *LocalNativeName, *FEmitHelper::GetCppName(ObjectClass), *OuterStr, *Object->GetName(), (int32)Object->GetFlags()));
+					, *LocalNativeName, *FEmitHelper::GetCppName(ObjectClass), *OuterStr, *Object->GetName(), ObjectFlags));
 			}
 		}
 		else
