Fix incorrect code generation at cook time when Blueprint nativization is enabled for objects instanced through the editor and assigned to an instanced object reference property.

Change summary:
- Modified FEmitDefaultValueHelper::HandleSpecialTypes() to emit runtime code to look for an existing instance of a default subobject only if the 'RF_DefaultSubObject' flag is set on the source instance. This works around instances that return 'true' for IsDefaultSubobject() but were not created with CreateDefaultSubobject() (e.g. - instanced in the editor UI).

diff --git a/Engine/Source/Developer/BlueprintCompilerCppBackend/Private/BlueprintCompilerCppBackendValueHelper.cpp b/Engine/Source/Developer/BlueprintCompilerCppBackend/Private/BlueprintCompilerCppBackendValueHelper.cpp
--- a/Engine/Source/Developer/BlueprintCompilerCppBackend/Private/BlueprintCompilerCppBackendValueHelper.cpp
+++ b/Engine/Source/Developer/BlueprintCompilerCppBackend/Private/BlueprintCompilerCppBackendValueHelper.cpp
@@ -699,7 +699,7 @@ FString FEmitDefaultValueHelper::HandleSpecialTypes(FEmitterLocalContext& Contex
 	{
 		if (Object)
 		{
-			const bool bIsDefaultSubobject = Object->IsDefaultSubobject();
+			const bool bIsDefaultSubobject = Object->IsDefaultSubobject() && Object->HasAllFlags(RF_DefaultSubObject);
 			const bool bIsInstancedReference = Property->HasAnyPropertyFlags(CPF_InstancedReference);
 
 			UClass* ObjectClassToUse = Context.GetFirstNativeOrConvertedClass(Class);
