Fix missing backend support for initializing a container variable (including unconnected function output pins) to a default (literal) value term in nativized Blueprint C++ code.

Change summary:
- Added a 'CoerceProperty' argument to FBlueprintCompilerCppBackEnd::TermToText() and FEmitHelper::LiteralTerm(). This is passed through as the LHS term's associated property when compiling 'KCST_Assignment' statements.
- Modified FEmitHelper::LiteralTerm() to emit an initialization list to C++. This occurs when the RHS of a variable assignment statement is a literal container term, which includes the initialization of K2Node_FunctionResult pins to a default (empty) value.
- Added FBlueprintCompilerCppBackEnd::TermToText() and FEmitHelper::LiteralTerm() methods that accept a struct value for specifying arguments rather than increasing the number of parameters with default values. Modified the implementations to match.

diff --git a/Engine/Source/Developer/BlueprintCompilerCppBackend/Private/BlueprintCompilerCppBackend.cpp b/Engine/Source/Developer/BlueprintCompilerCppBackend/Private/BlueprintCompilerCppBackend.cpp
--- a/Engine/Source/Developer/BlueprintCompilerCppBackend/Private/BlueprintCompilerCppBackend.cpp
+++ b/Engine/Source/Developer/BlueprintCompilerCppBackend/Private/BlueprintCompilerCppBackend.cpp
@@ -156,7 +156,13 @@ struct FSetterExpressionBuilder
 void FBlueprintCompilerCppBackend::EmitAssignmentStatment(FEmitterLocalContext& EmitterContext, FKismetFunctionContext& FunctionContext, FBlueprintCompiledStatement& Statement)
 {
 	check(Statement.LHS && Statement.RHS[0]);
-	const FString SourceExpression = TermToText(EmitterContext, Statement.RHS[0], ENativizedTermUsage::Getter);
+
+	FTermToTextParams SrcTermParams;
+	SrcTermParams.Term = Statement.RHS[0];
+	SrcTermParams.TermUsage = ENativizedTermUsage::Getter;
+	SrcTermParams.CoerceProperty = Statement.LHS->AssociatedVarProperty;
+
+	const FString SourceExpression = TermToText(EmitterContext, SrcTermParams);
 	FSetterExpressionBuilder SetterExpression(*this, EmitterContext, Statement.LHS);
 	FSafeContextScopedEmmitter SafeContextScope(EmitterContext, Statement.LHS->Context, *this);
 
@@ -908,20 +914,29 @@ FString FBlueprintCompilerCppBackend::EmitArrayGetByRef(FEmitterLocalContext& Em
 	return Result;
 }
 
-FString FBlueprintCompilerCppBackend::TermToText(FEmitterLocalContext& EmitterContext, const FBPTerminal* Term, const ENativizedTermUsage TermUsage, const bool bUseSafeContext, FString* EndCustomSetExpression)
+FString FBlueprintCompilerCppBackend::TermToText(FEmitterLocalContext& EmitterContext, const FTermToTextParams& Params)
 {
-	ensure((TermUsage != ENativizedTermUsage::Setter) || !bUseSafeContext);
-	ensure((TermUsage == ENativizedTermUsage::Setter) == (EndCustomSetExpression != nullptr));
-	if (EndCustomSetExpression)
+	const FBPTerminal* Term = Params.Term;
+	
+	ensure((Params.TermUsage != ENativizedTermUsage::Setter) || !Params.bUseSafeContext);
+	ensure((Params.TermUsage == ENativizedTermUsage::Setter) == (Params.EndCustomSetExpression != nullptr));
+	if (Params.EndCustomSetExpression)
 	{
-		EndCustomSetExpression->Reset();
+		Params.EndCustomSetExpression->Reset();
 	}
 
-	const bool bGetter = (TermUsage == ENativizedTermUsage::Getter);
+	const bool bGetter = (Params.TermUsage == ENativizedTermUsage::Getter);
 	const FString PSC_Self(TEXT("self"));
 	if (Term->bIsLiteral)
 	{
-		return FEmitHelper::LiteralTerm(EmitterContext, Term->Type, Term->Name, Term->ObjectLiteral, &Term->TextLiteral);
+		FEmitHelper::FLiteralTermParams LiteralTermParams;
+		LiteralTermParams.Type = Term->Type;
+		LiteralTermParams.CustomValue = Term->Name;
+		LiteralTermParams.LiteralText = Term->TextLiteral;
+		LiteralTermParams.LiteralObject = Term->ObjectLiteral;
+		LiteralTermParams.CoerceProperty = Params.CoerceProperty;
+
+		return FEmitHelper::LiteralTerm(EmitterContext, LiteralTermParams);
 	}
 	else if (Term->InlineGeneratedParameter)
 	{
@@ -953,11 +968,9 @@ FString FBlueprintCompilerCppBackend::TermToText(FEmitterLocalContext& EmitterCo
 				? EmitterContext.ExportCppDeclaration(Term->AssociatedVarProperty, EExportedDeclaration::Local, PropertyExportFlags, FEmitterLocalContext::EPropertyNameInDeclaration::Skip)
 				: FEmitHelper::PinTypeToNativeType(Term->Type);
 
-			const FString DefaultValueConstructor = (!Term->Type.IsContainer())
-				? FEmitHelper::LiteralTerm(EmitterContext, Term->Type, FString(), nullptr, &FText::GetEmpty())
-				: FString::Printf(TEXT("%s{}"), *CppType);
-
-			EmitterContext.AddLine(*FString::Printf(TEXT("%s %s = %s;"), *CppType, *DefaultValueVariable, *DefaultValueConstructor));
+			FEmitHelper::FLiteralTermParams DefaultValueTermParams;
+			DefaultValueTermParams.Type = Term->Type;
+			EmitterContext.AddLine(*FString::Printf(TEXT("%s %s = %s;"), *CppType, *DefaultValueVariable, *FEmitHelper::LiteralTerm(EmitterContext, DefaultValueTermParams)));
 
 			return DefaultValueVariable;
 		};
@@ -1005,7 +1018,7 @@ FString FBlueprintCompilerCppBackend::TermToText(FEmitterLocalContext& EmitterCo
 			bIsAccessible &= !Term->AssociatedVarProperty->HasAnyPropertyFlags(CPF_NativeAccessSpecifierPrivate | CPF_NativeAccessSpecifierProtected);
 			if (!bIsAccessible)
 			{
-				ResultPath = FEmitHelper::AccessInaccessibleProperty(EmitterContext, Term->AssociatedVarProperty, FString(), ContextStr, TEXT("&"), 0, TermUsage, EndCustomSetExpression);
+				ResultPath = FEmitHelper::AccessInaccessibleProperty(EmitterContext, Term->AssociatedVarProperty, FString(), ContextStr, TEXT("&"), 0, Params.TermUsage, Params.EndCustomSetExpression);
 			}
 			else
 			{
@@ -1044,7 +1057,7 @@ FString FBlueprintCompilerCppBackend::TermToText(FEmitterLocalContext& EmitterCo
 					ensure(ContextStr.IsEmpty());
 					ContextStr = TEXT("this");
 				}
-				ResultPath = FEmitHelper::AccessInaccessibleProperty(EmitterContext, Term->AssociatedVarProperty, FString(), ContextStr, FString(), 0, TermUsage, EndCustomSetExpression);
+				ResultPath = FEmitHelper::AccessInaccessibleProperty(EmitterContext, Term->AssociatedVarProperty, FString(), ContextStr, FString(), 0, Params.TermUsage, Params.EndCustomSetExpression);
 			}
 			else
 			{
@@ -1092,7 +1105,7 @@ FString FBlueprintCompilerCppBackend::TermToText(FEmitterLocalContext& EmitterCo
 			ResultPath = FString::Printf(TEXT("const_cast<%s>(%s)"), *CppType, *ResultPath);
 		}
 
-		const FString Conditions = bUseSafeContext ? FSafeContextScopedEmmitter::ValidationChain(EmitterContext, Term->Context, *this) : FString();
+		const FString Conditions = Params.bUseSafeContext ? FSafeContextScopedEmmitter::ValidationChain(EmitterContext, Term->Context, *this) : FString();
 		if (!Conditions.IsEmpty())
 		{
 			const FString DefaultValueVariable = GenerateDefaultLocalVariable(Term);
@@ -1121,23 +1134,25 @@ FString FBlueprintCompilerCppBackend::LatentFunctionInfoTermToText(FEmitterLocal
 
 	check(!FixupTermName.IsEmpty());
 
-	FString StructValues = Term->Name;
+	FEmitHelper::FLiteralTermParams LiteralTermParams;
+	LiteralTermParams.Type = Term->Type;
+	LiteralTermParams.CustomValue = Term->Name;
 
 	// Index 0 is always the ubergraph
 	const int32 TargetStateIndex = StateMapPerFunction[0].StatementToStateIndex(TargetLabel);
-	const int32 LinkageTermStartIdx = StructValues.Find(FixupTermName);
+	const int32 LinkageTermStartIdx = LiteralTermParams.CustomValue.Find(FixupTermName);
 	check(LinkageTermStartIdx != INDEX_NONE);
-	StructValues = StructValues.Replace(TEXT("-1"), *FString::FromInt(TargetStateIndex));
+	LiteralTermParams.CustomValue = LiteralTermParams.CustomValue.Replace(TEXT("-1"), *FString::FromInt(TargetStateIndex));
 
 	int32* ExecutionGroupPtr = UberGraphStatementToExecutionGroup.Find(TargetLabel);
 	if (ExecutionGroupPtr && UberGraphContext)
 	{
 		const FString OldExecutionFunctionName = UEdGraphSchema_K2::FN_ExecuteUbergraphBase.ToString() + TEXT("_") + UberGraphContext->Blueprint->GetName();
 		const FString NewExecutionFunctionName = OldExecutionFunctionName + FString::Printf(TEXT("_%d"), *ExecutionGroupPtr);
-		StructValues = StructValues.Replace(*OldExecutionFunctionName, *NewExecutionFunctionName);
+		LiteralTermParams.CustomValue = LiteralTermParams.CustomValue.Replace(*OldExecutionFunctionName, *NewExecutionFunctionName);
 	}
 
-	return FEmitHelper::LiteralTerm(EmitterContext, Term->Type, StructValues, nullptr);
+	return FEmitHelper::LiteralTerm(EmitterContext, LiteralTermParams);
 }
 
 bool FBlueprintCompilerCppBackend::InnerFunctionImplementation(FKismetFunctionContext& FunctionContext, FEmitterLocalContext& EmitterContext, int32 ExecutionGroup)
diff --git a/Engine/Source/Developer/BlueprintCompilerCppBackend/Private/BlueprintCompilerCppBackend.h b/Engine/Source/Developer/BlueprintCompilerCppBackend/Private/BlueprintCompilerCppBackend.h
--- a/Engine/Source/Developer/BlueprintCompilerCppBackend/Private/BlueprintCompilerCppBackend.h
+++ b/Engine/Source/Developer/BlueprintCompilerCppBackend/Private/BlueprintCompilerCppBackend.h
@@ -68,5 +68,32 @@ protected:
 	FString EmitArrayGetByRef(FEmitterLocalContext& EmitterContext, FBlueprintCompiledStatement& Statement);
 
 public:
-	FString TermToText(FEmitterLocalContext& EmitterContext, const FBPTerminal* Term, const ENativizedTermUsage TermUsage, const bool bUseSafeContext = true, FString* EndCustomSetExpression = nullptr);
+	struct FTermToTextParams
+	{
+		const FBPTerminal* Term;
+		ENativizedTermUsage TermUsage;
+		bool bUseSafeContext;
+		FString* EndCustomSetExpression;
+		const UProperty* CoerceProperty;
+
+		FTermToTextParams()
+			:Term(nullptr)
+			,TermUsage(ENativizedTermUsage::UnspecifiedOrReference)
+			,bUseSafeContext(true)
+			,EndCustomSetExpression(nullptr)
+			,CoerceProperty(nullptr)
+		{}
+	};
+
+	FString TermToText(FEmitterLocalContext& EmitterContext, const FTermToTextParams& Params);
+	FString TermToText(FEmitterLocalContext& EmitterContext, const FBPTerminal* Term, const ENativizedTermUsage TermUsage, const bool bUseSafeContext = true, FString* EndCustomSetExpression = nullptr)
+	{
+		FTermToTextParams Params;
+		Params.Term = Term;
+		Params.TermUsage = TermUsage;
+		Params.bUseSafeContext = bUseSafeContext;
+		Params.EndCustomSetExpression = EndCustomSetExpression;
+
+		return TermToText(EmitterContext, Params);
+	}
 };
diff --git a/Engine/Source/Developer/BlueprintCompilerCppBackend/Private/BlueprintCompilerCppBackendUtils.cpp b/Engine/Source/Developer/BlueprintCompilerCppBackend/Private/BlueprintCompilerCppBackendUtils.cpp
--- a/Engine/Source/Developer/BlueprintCompilerCppBackend/Private/BlueprintCompilerCppBackendUtils.cpp
+++ b/Engine/Source/Developer/BlueprintCompilerCppBackend/Private/BlueprintCompilerCppBackendUtils.cpp
@@ -1021,22 +1021,193 @@ FString FEmitHelper::FloatToString(float Value)
 	return FString::Printf(TEXT("%f"), Value);
 }
 
-FString FEmitHelper::LiteralTerm(FEmitterLocalContext& EmitterContext, const FEdGraphPinType& Type, const FString& CustomValue, UObject* LiteralObject, const FText* OptionalTextLiteral)
+FString FEmitHelper::LiteralTerm(FEmitterLocalContext& EmitterContext, const FLiteralTermParams& Params)
 {
 	auto Schema = GetDefault<UEdGraphSchema_K2>();
 
-	if (UEdGraphSchema_K2::PC_String == Type.PinCategory)
+	class FImportTextErrorContext : public FStringOutputDevice
+	{
+	public:
+		int32 NumErrors;
+
+		FImportTextErrorContext() : FStringOutputDevice(), NumErrors(0) {}
+
+		virtual void Serialize(const TCHAR* V, ELogVerbosity::Type Verbosity, const class FName& Category) override
+		{
+			if (ELogVerbosity::Type::Error == Verbosity)
+			{
+				NumErrors++;
+			}
+			FStringOutputDevice::Serialize(V, Verbosity, Category);
+		}
+	};
+
+	const FEdGraphPinType& Type = Params.Type;
+	const FString& CustomValue = Params.CustomValue;
+
+	if (Type.IsContainer())
+	{
+		FString ContainerInitializerList;
+		FImportTextErrorContext ImportError;
+
+		if (const UArrayProperty* ArrayProperty = Cast<UArrayProperty>(Params.CoerceProperty))
+		{
+			FScriptArray ScriptArray;
+			if (!ArrayProperty->ImportText(*CustomValue, &ScriptArray, PPF_None, nullptr, &ImportError))
+			{
+				UE_LOG(LogK2Compiler, Error, TEXT("FEmitHelper::LiteralTerm cannot parse array value \"%s\" error: %s class: %s"), *CustomValue, *ImportError, *GetPathNameSafe(EmitterContext.GetCurrentlyGeneratedClass()));
+			}
+
+			const int32 NumElements = ScriptArray.Num();
+			FScriptArrayHelper ScriptArrayHelper(ArrayProperty, &ScriptArray);
+
+			FLiteralTermParams InnerTermParams;
+			Schema->ConvertPropertyToPinType(ArrayProperty->Inner, InnerTermParams.Type);
+
+			const UTextProperty* InnerTextProperty = Cast<UTextProperty>(ArrayProperty->Inner);
+			const UObjectPropertyBase* InnerObjectProperty = InnerTextProperty ? nullptr : Cast<UObjectPropertyBase>(ArrayProperty->Inner);
+
+			for (int32 ElementIdx = 0; ElementIdx < NumElements; ++ElementIdx)
+			{
+				uint8* ValuePtr = ScriptArrayHelper.GetRawPtr(ElementIdx);
+				if (ArrayProperty->Inner->ExportText_Direct(InnerTermParams.CustomValue, ValuePtr, ValuePtr, nullptr, PPF_None))
+				{
+					if (InnerTextProperty)
+					{
+						InnerTermParams.LiteralText = InnerTextProperty->GetPropertyValue(ValuePtr);
+						InnerTermParams.CustomValue = InnerTermParams.LiteralText.ToString();
+					}
+					else if (InnerObjectProperty)
+					{
+						InnerTermParams.LiteralObject = InnerObjectProperty->GetObjectPropertyValue(ValuePtr);
+					}
+
+					ContainerInitializerList += LiteralTerm(EmitterContext, InnerTermParams);
+					if (ElementIdx < NumElements - 1)
+					{
+						ContainerInitializerList += TEXT(", ");
+					}
+				}
+			}
+		}
+		else if (const USetProperty* SetProperty = Cast<USetProperty>(Params.CoerceProperty))
+		{
+			FScriptSet ScriptSet;
+			if (!CustomValue.IsEmpty())	// unlike UArrayProperty, USetProperty::ImportText() doesn't allow empty values to pass, so we check for that here.
+			{
+				if (!SetProperty->ImportText(*CustomValue, &ScriptSet, PPF_None, nullptr, &ImportError))
+				{
+					UE_LOG(LogK2Compiler, Error, TEXT("FEmitHelper::LiteralTerm cannot parse set value \"%s\" error: %s class: %s"), *CustomValue, *ImportError, *GetPathNameSafe(EmitterContext.GetCurrentlyGeneratedClass()));
+				}
+			}
+
+			const int32 NumElements = ScriptSet.Num();
+			FScriptSetHelper ScriptSetHelper(SetProperty, &ScriptSet);
+
+			FLiteralTermParams ElementTermParams;
+			Schema->ConvertPropertyToPinType(SetProperty->ElementProp, ElementTermParams.Type);
+
+			const UTextProperty* ElementTextProperty = Cast<UTextProperty>(SetProperty->ElementProp);
+			const UObjectPropertyBase* ElementObjectProperty = ElementTextProperty ? nullptr : Cast<UObjectPropertyBase>(SetProperty->ElementProp);
+
+			for (int32 ElementIdx = 0; ElementIdx < NumElements; ++ElementIdx)
+			{
+				uint8* ValuePtr = ScriptSetHelper.GetElementPtr(ElementIdx);
+				if (SetProperty->ElementProp->ExportText_Direct(ElementTermParams.CustomValue, ValuePtr, ValuePtr, nullptr, PPF_None))
+				{
+					if (ElementTextProperty)
+					{
+						ElementTermParams.LiteralText = ElementTextProperty->GetPropertyValue(ValuePtr);
+						ElementTermParams.CustomValue = ElementTermParams.LiteralText.ToString();
+					}
+					else if (ElementObjectProperty)
+					{
+						ElementTermParams.LiteralObject = ElementObjectProperty->GetObjectPropertyValue(ValuePtr);
+					}
+
+					ContainerInitializerList += LiteralTerm(EmitterContext, ElementTermParams);
+					if (ElementIdx < NumElements - 1)
+					{
+						ContainerInitializerList += TEXT(", ");
+					}
+				}
+			}
+		}
+		else if (const UMapProperty* MapProperty = Cast<UMapProperty>(Params.CoerceProperty))
+		{
+			FScriptSet ScriptMap;
+			if (!CustomValue.IsEmpty())	// unlike UArrayProperty, UMapProperty::ImportText() doesn't allow empty values to pass, so we check for that here.
+			{
+				if (!MapProperty->ImportText(*CustomValue, &ScriptMap, PPF_None, nullptr, &ImportError))
+				{
+					UE_LOG(LogK2Compiler, Error, TEXT("FEmitHelper::LiteralTerm cannot parse map value \"%s\" error: %s class: %s"), *CustomValue, *ImportError, *GetPathNameSafe(EmitterContext.GetCurrentlyGeneratedClass()));
+				}
+			}
+
+			const int32 NumElements = ScriptMap.Num();
+			FScriptMapHelper ScriptMapHelper(MapProperty, &ScriptMap);
+
+			FLiteralTermParams KeyTermParams, ValueTermParams;
+			Schema->ConvertPropertyToPinType(MapProperty->KeyProp, KeyTermParams.Type);
+			Schema->ConvertPropertyToPinType(MapProperty->ValueProp, ValueTermParams.Type);
+
+			const UTextProperty* KeyTextProperty = Cast<UTextProperty>(MapProperty->KeyProp);
+			const UTextProperty* ValueTextProperty = Cast<UTextProperty>(MapProperty->ValueProp);
+			const UObjectPropertyBase* KeyObjectProperty = KeyTextProperty ? nullptr : Cast<UObjectPropertyBase>(MapProperty->KeyProp);
+			const UObjectPropertyBase* ValueObjectProperty = ValueTextProperty ? nullptr : Cast<UObjectPropertyBase>(MapProperty->ValueProp);
+
+			for (int32 ElementIdx = 0, SparseIdx = 0; ElementIdx < NumElements; ++SparseIdx)
+			{
+				if (ScriptMap.IsValidIndex(SparseIdx))
+				{
+					uint8* KeyPtr = ScriptMapHelper.GetKeyPtr(SparseIdx);
+					if (MapProperty->KeyProp->ExportText_Direct(KeyTermParams.CustomValue, KeyPtr, KeyPtr, nullptr, PPF_None))
+					{
+						if (KeyTextProperty)
+						{
+							KeyTermParams.LiteralText = KeyTextProperty->GetPropertyValue(KeyPtr);
+							KeyTermParams.CustomValue = KeyTermParams.LiteralText.ToString();
+						}
+						else if (KeyObjectProperty)
+						{
+							KeyTermParams.LiteralObject = KeyObjectProperty->GetObjectPropertyValue(KeyPtr);
+						}
+
+						uint8* ValuePtr = ScriptMapHelper.GetValuePtr(SparseIdx);
+						if (MapProperty->ValueProp->ExportText_Direct(ValueTermParams.CustomValue , ValuePtr, ValuePtr, nullptr, PPF_None))
+						{
+							if (ValueTextProperty)
+							{
+								ValueTermParams.LiteralText = ValueTextProperty->GetPropertyValue(ValuePtr);
+								ValueTermParams.CustomValue = ValueTermParams.LiteralText.ToString();
+							}
+							else if (ValueObjectProperty)
+							{
+								ValueTermParams.LiteralObject = ValueObjectProperty->GetObjectPropertyValue(ValuePtr);
+							}
+
+							ContainerInitializerList += FString::Printf(TEXT("{%s, %s}"), *LiteralTerm(EmitterContext, KeyTermParams), *LiteralTerm(EmitterContext, ValueTermParams));
+							if (ElementIdx < NumElements - 1)
+							{
+								ContainerInitializerList += TEXT(", ");
+							}
+						}
+					}
+
+					++ElementIdx;
+				}
+			}
+		}
+
+		return FString::Printf(TEXT("{%s}"), *ContainerInitializerList);
+	}
+	else if (UEdGraphSchema_K2::PC_String == Type.PinCategory)
 	{
 		return FString::Printf(TEXT("FString(%s)"), *UStrProperty::ExportCppHardcodedText(CustomValue, EmitterContext.DefaultTarget->Indent));
 	}
 	else if (UEdGraphSchema_K2::PC_Text == Type.PinCategory)
 	{
-		ensure(OptionalTextLiteral);
-		if (OptionalTextLiteral)
-		{
-			return UTextProperty::GenerateCppCodeForTextValue(*OptionalTextLiteral, FString());
-		}
-		return FString::Printf(TEXT("FText::FromString(%s)"), *UStrProperty::ExportCppHardcodedText(CustomValue, EmitterContext.DefaultTarget->Indent));
+		return UTextProperty::GenerateCppCodeForTextValue(Params.LiteralText, FString());
 	}
 	else if (UEdGraphSchema_K2::PC_Float == Type.PinCategory)
 	{
@@ -1151,23 +1322,6 @@ FString FEmitHelper::LiteralTerm(FEmitterLocalContext& EmitterContext, const FEd
 				FStructOnScope StructOnScope(StructType);
 				StructType->InitializeDefaultValue(StructOnScope.GetStructMemory()); // after cl#3098294 only delta (of structure data) against the default value is stored in string. So we need to explicitly provide default values before serialization.
 
-				class FImportTextErrorContext : public FStringOutputDevice
-				{
-				public:
-					int32 NumErrors;
-
-					FImportTextErrorContext() : FStringOutputDevice(), NumErrors(0) {}
-
-					virtual void Serialize(const TCHAR* V, ELogVerbosity::Type Verbosity, const class FName& Category) override
-					{
-						if (ELogVerbosity::Type::Error == Verbosity)
-						{
-							NumErrors++;
-						}
-						FStringOutputDevice::Serialize(V, Verbosity, Category);
-					}
-				};
-
 				FImportTextErrorContext ImportError;
 				const TCHAR* EndOfParsedBuff = StructType->ImportText(*CustomValue, StructOnScope.GetStructMemory(), nullptr, PPF_None, &ImportError, TEXT("FEmitHelper::LiteralTerm"));
 				if (!EndOfParsedBuff || ImportError.NumErrors)
@@ -1205,14 +1359,14 @@ FString FEmitHelper::LiteralTerm(FEmitterLocalContext& EmitterContext, const FEd
 	}
 	else if (UEdGraphSchema_K2::PC_Class == Type.PinCategory)
 	{
-		if (auto FoundClass = Cast<const UClass>(LiteralObject))
+		if (auto FoundClass = Cast<const UClass>(Params.LiteralObject))
 		{
-			const FString MappedObject = EmitterContext.FindGloballyMappedObject(LiteralObject, UClass::StaticClass());
+			const FString MappedObject = EmitterContext.FindGloballyMappedObject(Params.LiteralObject, UClass::StaticClass());
 			if (!MappedObject.IsEmpty())
 			{
 				return MappedObject;
 			}
-			return FString::Printf(TEXT("LoadClass<UClass>(nullptr, TEXT(\"%s\"), nullptr, 0, nullptr)"), *(LiteralObject->GetPathName().ReplaceCharWithEscapedChar()));
+			return FString::Printf(TEXT("LoadClass<UClass>(nullptr, TEXT(\"%s\"), nullptr, 0, nullptr)"), *(Params.LiteralObject->GetPathName().ReplaceCharWithEscapedChar()));
 		}
 		return FString(TEXT("((UClass*)nullptr)"));
 	}
@@ -1236,9 +1390,9 @@ FString FEmitHelper::LiteralTerm(FEmitterLocalContext& EmitterContext, const FEd
 	{
 		UClass* FoundClass = Cast<UClass>(Type.PinSubCategoryObject.Get());
 		UClass* ObjectClassToUse = FoundClass ? EmitterContext.GetFirstNativeOrConvertedClass(FoundClass) : UObject::StaticClass();
-		if (LiteralObject)
+		if (Params.LiteralObject)
 		{
-			const FString MappedObject = EmitterContext.FindGloballyMappedObject(LiteralObject, ObjectClassToUse, true);
+			const FString MappedObject = EmitterContext.FindGloballyMappedObject(Params.LiteralObject, ObjectClassToUse, true);
 			if (!MappedObject.IsEmpty())
 			{
 				return MappedObject;
@@ -1249,25 +1403,12 @@ FString FEmitHelper::LiteralTerm(FEmitterLocalContext& EmitterContext, const FEd
 	}
 	else if (UEdGraphSchema_K2::PC_Interface == Type.PinCategory)
 	{
-		if (!LiteralObject && CustomValue.IsEmpty())
+		if (!Params.LiteralObject && CustomValue.IsEmpty())
 		{
 			return FString(TEXT("nullptr"));
 		}
 	}
-	/*
-	else if (CoerceProperty->IsA(UInterfaceProperty::StaticClass()))
-	{
-		if (Term->Type.PinSubCategory == PSC_Self)
-		{
-			return TEXT("this");
-		}
-		else
-		{
-			ensureMsgf(false, TEXT("It is not possible to express this interface property as a literal value!"));
-			return Term->Name;
-		}
-	}
-	*/
+
 	ensureMsgf(false, TEXT("It is not possible to express this type as a literal value!"));
 	return CustomValue;
 }
diff --git a/Engine/Source/Developer/BlueprintCompilerCppBackend/Private/BlueprintCompilerCppBackendUtils.h b/Engine/Source/Developer/BlueprintCompilerCppBackend/Private/BlueprintCompilerCppBackendUtils.h
--- a/Engine/Source/Developer/BlueprintCompilerCppBackend/Private/BlueprintCompilerCppBackendUtils.h
+++ b/Engine/Source/Developer/BlueprintCompilerCppBackend/Private/BlueprintCompilerCppBackendUtils.h
@@ -287,7 +287,22 @@ struct FEmitHelper
 
 	static void EmitLifetimeReplicatedPropsImpl(FEmitterLocalContext& EmitterContext);
 
-	static FString LiteralTerm(FEmitterLocalContext& EmitterContext, const FEdGraphPinType& Type, const FString& CustomValue, UObject* LiteralObject, const FText* OptionalTextLiteral = nullptr);
+	struct FLiteralTermParams
+	{
+		FEdGraphPinType Type;
+		FString CustomValue;
+		FText LiteralText;
+		UObject* LiteralObject;
+		const UProperty* CoerceProperty;
+
+		FLiteralTermParams()
+			:LiteralObject(nullptr)
+			,CoerceProperty(nullptr)
+		{
+		}
+	};
+
+	static FString LiteralTerm(FEmitterLocalContext& EmitterContext, const FLiteralTermParams& Params);
 
 	static FString PinTypeToNativeType(const FEdGraphPinType& InType);
 
