Fix for nativized builds failing on a VirtualProductionUtilities plugin Blueprint asset dependency.

#rnx

diff --git a/Engine/Config/BaseEditor.ini b/Engine/Config/BaseEditor.ini
--- a/Engine/Config/BaseEditor.ini
+++ b/Engine/Config/BaseEditor.ini
@@ -450,7 +450,8 @@ bNativizeAnimBPOnlyWhenNonReducibleFuncitons=false
 +ExcludedAssets=/Engine/Tutorial/SubEditors/TutorialAssets/Character/TutorialCharacter
 +ExcludedAssets=/Engine/Tutorial/SubEditors/TutorialAssets/Character/TutorialTPP_AnimBlueprint
 
-; +ExcludedFolderPaths=...
+; Currently excluded due to an editor-only asset dependency
++ExcludedAssets=/VirtualProductionUtilities/UserWidgets/BookmarkWidget
 
 ; To add a module dependency to the nativized module in both client and server use:
 ; +AdditionalPublicDependencyModuleNames=
diff --git a/Engine/Source/Developer/BlueprintCompilerCppBackend/Private/BlueprintCompilerCppBackendValueHelper.cpp b/Engine/Source/Developer/BlueprintCompilerCppBackend/Private/BlueprintCompilerCppBackendValueHelper.cpp
--- a/Engine/Source/Developer/BlueprintCompilerCppBackend/Private/BlueprintCompilerCppBackendValueHelper.cpp
+++ b/Engine/Source/Developer/BlueprintCompilerCppBackend/Private/BlueprintCompilerCppBackendValueHelper.cpp
@@ -1336,7 +1336,7 @@ void FEmitDefaultValueHelper::AddStaticFunctionsForDependencies(FEmitterLocalCon
 		ensure(InAsset);
 		if (InAsset && IsEditorOnlyObject(InAsset))
 		{
-			UE_LOG(LogK2Compiler, Warning, TEXT("Nativized %d depends on editor only asset: %s")
+			UE_LOG(LogK2Compiler, Warning, TEXT("Nativized %s depends on editor only asset: %s")
 				, (OriginalClass ? *OriginalClass->GetPathName() : *CppTypeName)
 				,*InAsset->GetPathName());
 			OptionalComment = TEXT("Editor Only asset");
