Fix conversion of a TWeakObjectPtr property value when it references an instanced subobject in a nativized Blueprint class.

diff --git a/Engine/Source/Developer/BlueprintCompilerCppBackend/Private/BlueprintCompilerCppBackendValueHelper.cpp b/Engine/Source/Developer/BlueprintCompilerCppBackend/Private/BlueprintCompilerCppBackendValueHelper.cpp
--- a/Engine/Source/Developer/BlueprintCompilerCppBackend/Private/BlueprintCompilerCppBackendValueHelper.cpp
+++ b/Engine/Source/Developer/BlueprintCompilerCppBackend/Private/BlueprintCompilerCppBackendValueHelper.cpp
@@ -786,6 +786,10 @@ FString FEmitDefaultValueHelper::HandleSpecialTypes(FEmitterLocalContext& Contex
 	{
 		Result = HandleObjectValueLambda(ObjectProperty->GetPropertyValue(ValuePtr), ObjectProperty->PropertyClass);
 	}
+	else if (const UWeakObjectProperty* WeakObjectProperty = Cast<UWeakObjectProperty>(Property))
+	{
+		Result = HandleObjectValueLambda(WeakObjectProperty->GetObjectPropertyValue(ValuePtr), WeakObjectProperty->PropertyClass);
+	}
 	else if (const UInterfaceProperty* InterfaceProperty = Cast<UInterfaceProperty>(Property))
 	{
 		Result = HandleObjectValueLambda(InterfaceProperty->GetPropertyValue(ValuePtr).GetObject(), InterfaceProperty->InterfaceClass);
