Fix replication of a child actor component's internal archetype object when both the archetype and the outer Blueprint class are both nativized.

diff --git a/Engine/Source/Developer/BlueprintCompilerCppBackend/Private/BlueprintCompilerCppBackendValueHelper.cpp b/Engine/Source/Developer/BlueprintCompilerCppBackend/Private/BlueprintCompilerCppBackendValueHelper.cpp
--- a/Engine/Source/Developer/BlueprintCompilerCppBackend/Private/BlueprintCompilerCppBackendValueHelper.cpp
+++ b/Engine/Source/Developer/BlueprintCompilerCppBackend/Private/BlueprintCompilerCppBackendValueHelper.cpp
@@ -2036,6 +2036,7 @@ FString FEmitDefaultValueHelper::HandleClassSubobject(FEmitterLocalContext& Cont
 		LocalNativeName = Context.GenerateUniqueLocalName();
 		Context.AddClassSubObject_InConstructor(Object, LocalNativeName);
 		UClass* ObjectClass = Object->GetClass();
+		const int32 ObjectFlags = (int32)Object->GetFlags();
 		const FString ActualClass = Context.FindGloballyMappedObject(ObjectClass, UClass::StaticClass());
 		const FString NativeType = FEmitHelper::GetCppName(Context.GetFirstNativeOrConvertedClass(ObjectClass));
 		if(!ObjectClass->IsNative())
@@ -2044,12 +2045,13 @@ FString FEmitDefaultValueHelper::HandleClassSubobject(FEmitterLocalContext& Cont
 			Context.AddLine(FString::Printf(TEXT("%s::StaticClass()->GetDefaultObject();"), *NativeType));
 		}
 		Context.AddLine(FString::Printf(
-			TEXT("auto %s = NewObject<%s>(%s, %s, TEXT(\"%s\"));")
+			TEXT("auto %s = NewObject<%s>(%s, %s, TEXT(\"%s\"), (EObjectFlags)0x%08x);")
 			, *LocalNativeName
 			, *NativeType
 			, *OuterStr
 			, *ActualClass
-			, *Object->GetName().ReplaceCharWithEscapedChar()));
+			, *Object->GetName().ReplaceCharWithEscapedChar()
+			, ObjectFlags));
 		if (bAddAsSubobjectOfClass)
 		{
 			Context.RegisterClassSubobject(Object, ListOfSubobjectsType);
