Fix broken skeletal animation in a nativized, cooked build (regression).

Change summary:
- Added FEmitDefaultValueHelper::EPropertyGenerationControlFlags.
- Replaced 'bool' control parameters on FEmitDefaultValueHelper::OuterGenerate/InnerGenerate call sites with EPropertyGenerationControlFlags.
    - Primary reason for this change was so that I could allow transient properties to pass through in certain cases without adding another 'bool' parameter to these functions.
    - This flag may also serve as a foundation for eventually fixing UE-54921.
- Fixed a bug in FEmitterLocalContext::FindGloballyMappedObject() where we weren't properly recognizing UFunction ptr values as a UField with a UStruct owner. This was causing us to emit code that instanced a duplicate UFunction as a class-owned subobject instead!
- Modified FKismetCompilerContext::CompileFunctions() to call PostCDOCompiled() when we're not in the BPCM code path. This fixes a bug where some of the new AnimBP fixup code wasn't happening after duplicating a BP for conversion to C++.
- Added transient UPROPERTY markup to non-serialized fields in FAnimBlueprintFunction. This ensures that we will emit initialization code for these fields in a converted AnimBP class if they contain non-default values.
- Modified FBackendHelperAnim::CreateAnimClassData() to direct FEmitDefaultValueHelper::OuterGenerate() to emit initialization code for transient properties.
- Removed UAnimClassData::InitGraphExposedInputs(). This is now being handled by code that's emitted to the individual node initialization functions in FBackendHelperAnim::AddAnimNodeInitializationFunction(), and it also fixes a thread safety issue with async loading since it no longer involves a FindFunction() call.

#rnx

diff --git a/Engine/Source/Developer/BlueprintCompilerCppBackend/Private/BlueprintCompilerCppBackendAnim.cpp b/Engine/Source/Developer/BlueprintCompilerCppBackend/Private/BlueprintCompilerCppBackendAnim.cpp
--- a/Engine/Source/Developer/BlueprintCompilerCppBackend/Private/BlueprintCompilerCppBackendAnim.cpp
+++ b/Engine/Source/Developer/BlueprintCompilerCppBackend/Private/BlueprintCompilerCppBackendAnim.cpp
@@ -33,7 +33,8 @@ void FBackendHelperAnim::CreateAnimClassData(FEmitterLocalContext& Context)
 			FEmitDefaultValueHelper::OuterGenerate(Context, Property, LocalNativeName
 				, reinterpret_cast<const uint8*>(AnimClassData)
 				, reinterpret_cast<const uint8*>(ObjectArchetype)
-				, FEmitDefaultValueHelper::EPropertyAccessOperator::Pointer);
+				, FEmitDefaultValueHelper::EPropertyAccessOperator::Pointer
+				, FEmitDefaultValueHelper::EPropertyGenerationControlFlags::AllowTransient);
 		}
 
 		Context.AddLine(FString::Printf(TEXT("InDynamicClass->%s = %s;"), GET_MEMBER_NAME_STRING_CHECKED(UDynamicClass, AnimClassImplementation), *LocalNativeName));
@@ -76,7 +77,23 @@ void FBackendHelperAnim::AddAnimNodeInitializationFunction(FEmitterLocalContext&
 			Context.Body.AddLine(TEXT("{"));
 			Context.Body.IncreaseIndent();
 
-			FEmitDefaultValueHelper::OuterGenerate(Context, InProperty, TEXT(""), reinterpret_cast<const uint8*>(InCDO), bInNewProperty ? nullptr : reinterpret_cast<const uint8*>(InParentCDO), FEmitDefaultValueHelper::EPropertyAccessOperator::None, true);
+			FEmitDefaultValueHelper::OuterGenerate(Context, InProperty, TEXT(""), reinterpret_cast<const uint8*>(InCDO), bInNewProperty ? nullptr : reinterpret_cast<const uint8*>(InParentCDO), FEmitDefaultValueHelper::EPropertyAccessOperator::None, FEmitDefaultValueHelper::EPropertyGenerationControlFlags::AllowProtected);
+			
+			// anim nodes constructed, finish anim node initialization:
+			if (UAnimBlueprintGeneratedClass* AnimClass = Cast<UAnimBlueprintGeneratedClass>(Context.GetCurrentlyGeneratedClass()))
+			{
+				for (int32 i = 0; i < AnimClass->EvaluateGraphExposedInputs.Num(); ++i)
+				{
+					if (AnimClass->EvaluateGraphExposedInputs[i].ValueHandlerNodeProperty == InProperty)
+					{
+						const FString ClassName = FEmitHelper::GetCppName(Context.GetCurrentlyGeneratedClass());
+						const FString MemberName = FEmitHelper::GetCppName(InProperty);
+						Context.Body.AddLine(FString::Printf(TEXT("%s.SetExposedValueHandler(&CastChecked<UAnimClassData>(CastChecked<UDynamicClass>(%s::StaticClass())->%s)->GetExposedValueHandlers()[%d]);"), *MemberName, *ClassName, GET_MEMBER_NAME_STRING_CHECKED(UDynamicClass, AnimClassImplementation), i));
+
+						break;
+					}
+				}
+			}
 
 			Context.Body.DecreaseIndent();
 			Context.Body.AddLine(TEXT("}"));
@@ -104,8 +121,4 @@ void FBackendHelperAnim::AddAllAnimNodesInitializationFunction(FEmitterLocalCont
 void FBackendHelperAnim::AddAllAnimNodesInitializationFunctionCall(FEmitterLocalContext& Context)
 {
 	Context.Body.AddLine(TEXT("__InitAllAnimNodes();"));
-	
-	// anim nodes constructed, finish anim node initialization:
-	FString Class = FEmitHelper::GetCppName(Context.GetCurrentlyGeneratedClass());
-	Context.Body.AddLine(FString::Printf(TEXT("CastChecked<UAnimClassData>(CastChecked<UDynamicClass>(%s::StaticClass())->%s)->InitGraphExposedInputs(this);"), *Class, GET_MEMBER_NAME_STRING_CHECKED(UDynamicClass, AnimClassImplementation)));
 }
\ No newline at end of file
diff --git a/Engine/Source/Developer/BlueprintCompilerCppBackend/Private/BlueprintCompilerCppBackendUMG.cpp b/Engine/Source/Developer/BlueprintCompilerCppBackend/Private/BlueprintCompilerCppBackendUMG.cpp
--- a/Engine/Source/Developer/BlueprintCompilerCppBackend/Private/BlueprintCompilerCppBackendUMG.cpp
+++ b/Engine/Source/Developer/BlueprintCompilerCppBackend/Private/BlueprintCompilerCppBackendUMG.cpp
@@ -83,7 +83,7 @@ void FBackendHelperUMG::EmitWidgetInitializationFunctions(FEmitterLocalContext&
 			const FString Target = InContext.ExportCppDeclaration(InProperty, EExportedDeclaration::Local, CppTemplateTypeFlags, FEmitterLocalContext::EPropertyNameInDeclaration::Skip);
 
 			InContext.AddLine(FString::Printf(TEXT("%s %s;"), *Target, *NativeName));
-			FEmitDefaultValueHelper::InnerGenerate(InContext, InProperty, NativeName, DataPtr, nullptr, true);
+			FEmitDefaultValueHelper::InnerGenerate(InContext, InProperty, NativeName, DataPtr, nullptr);
 			return NativeName;
 		};
 
diff --git a/Engine/Source/Developer/BlueprintCompilerCppBackend/Private/BlueprintCompilerCppBackendUtils.cpp b/Engine/Source/Developer/BlueprintCompilerCppBackend/Private/BlueprintCompilerCppBackendUtils.cpp
--- a/Engine/Source/Developer/BlueprintCompilerCppBackend/Private/BlueprintCompilerCppBackendUtils.cpp
+++ b/Engine/Source/Developer/BlueprintCompilerCppBackend/Private/BlueprintCompilerCppBackendUtils.cpp
@@ -137,8 +137,19 @@ FString FEmitterLocalContext::FindGloballyMappedObject(const UObject* Object, co
 	}
 
 	{
+		// Need special-case handling for UFunction-type fields (can't use GetOwnerStruct).
+		auto GetFieldOwnerStruct = [](const UField* InField) -> UStruct*
+		{
+			if (InField->IsA<UFunction>())
+			{
+				return InField->GetOwnerClass();
+			}
+
+			return InField->GetOwnerStruct();
+		};
+
 		const UField* Field = Cast<UField>(Object);
-		const UStruct* FieldOwnerStruct = Field ? Field->GetOwnerStruct() : nullptr;
+		const UStruct* FieldOwnerStruct = Field ? GetFieldOwnerStruct(Field) : nullptr;
 		if (FieldOwnerStruct && (Field != FieldOwnerStruct))
 		{
 			ensure(Field == FindField<UField>(FieldOwnerStruct, Field->GetFName()));
diff --git a/Engine/Source/Developer/BlueprintCompilerCppBackend/Private/BlueprintCompilerCppBackendUtils.h b/Engine/Source/Developer/BlueprintCompilerCppBackend/Private/BlueprintCompilerCppBackendUtils.h
--- a/Engine/Source/Developer/BlueprintCompilerCppBackend/Private/BlueprintCompilerCppBackendUtils.h
+++ b/Engine/Source/Developer/BlueprintCompilerCppBackend/Private/BlueprintCompilerCppBackendUtils.h
@@ -350,12 +350,25 @@ struct FEmitDefaultValueHelper
 		Dot,
 	};
 
+	/** Allows callers to control the behavior of emitting code to initialize properties. */
+	enum class EPropertyGenerationControlFlags
+	{
+		/** Normal operation (default). */
+		None = 0,
+		/** Allow transient properties to be emitted. By default, transient properties are skipped. */
+		AllowTransient = 1 << 0,
+		/** Allow protected properties to be treated as an accessible property. By default, protected properties will use the inaccessible property path, which emits some extra code. */
+		AllowProtected = 1 << 1,
+		/** Emit code to construct the value before assigning it. By default, it's assumed that the value has already been constructed. */
+		IncludeFirstConstructionLine = 1 << 2
+	};
+
 	// OuterPath ends context/outer name (or empty, if the scope is "this")
-	static void OuterGenerate(FEmitterLocalContext& Context, const UProperty* Property, const FString& OuterPath, const uint8* DataContainer, const uint8* OptionalDefaultDataContainer, EPropertyAccessOperator AccessOperator, bool bAllowProtected = false);
+	static void OuterGenerate(FEmitterLocalContext& Context, const UProperty* Property, const FString& OuterPath, const uint8* DataContainer, const uint8* OptionalDefaultDataContainer, EPropertyAccessOperator AccessOperator, EPropertyGenerationControlFlags ControlFlags = EPropertyGenerationControlFlags::None);
 
 
 	// PathToMember ends with variable name
-	static void InnerGenerate(FEmitterLocalContext& Context, const UProperty* Property, const FString& PathToMember, const uint8* ValuePtr, const uint8* DefaultValuePtr, bool bWithoutFirstConstructionLine = false);
+	static void InnerGenerate(FEmitterLocalContext& Context, const UProperty* Property, const FString& PathToMember, const uint8* ValuePtr, const uint8* DefaultValuePtr, EPropertyGenerationControlFlags ControlFlags = EPropertyGenerationControlFlags::None);
 
 	// Creates the subobject (of class) returns it's native local name, 
 	// returns empty string if cannot handle
@@ -381,6 +394,8 @@ private:
 		, bool bBlockRecursion);
 };
 
+ENUM_CLASS_FLAGS(FEmitDefaultValueHelper::EPropertyGenerationControlFlags);
+
 struct FBackendHelperUMG
 {
 	static void WidgetFunctionsInHeader(FEmitterLocalContext& Context);
diff --git a/Engine/Source/Developer/BlueprintCompilerCppBackend/Private/BlueprintCompilerCppBackendValueHelper.cpp b/Engine/Source/Developer/BlueprintCompilerCppBackend/Private/BlueprintCompilerCppBackendValueHelper.cpp
--- a/Engine/Source/Developer/BlueprintCompilerCppBackend/Private/BlueprintCompilerCppBackendValueHelper.cpp
+++ b/Engine/Source/Developer/BlueprintCompilerCppBackend/Private/BlueprintCompilerCppBackendValueHelper.cpp
@@ -38,11 +38,17 @@ void FEmitDefaultValueHelper::OuterGenerate(FEmitterLocalContext& Context
 	, const uint8* DataContainer
 	, const uint8* OptionalDefaultDataContainer
 	, EPropertyAccessOperator AccessOperator
-	, bool bAllowProtected)
+	, EPropertyGenerationControlFlags ControlFlags)
 {
 	check(Property);
 
-	if (Property->HasAnyPropertyFlags(CPF_EditorOnly | CPF_Transient))
+	const bool bAllowTransient = !!(ControlFlags & EPropertyGenerationControlFlags::AllowTransient);
+	const bool bAllowProtected = !!(ControlFlags & EPropertyGenerationControlFlags::AllowProtected);
+
+	// Don't allow these flags to propagate through to nested calls.
+	ControlFlags &= ~EPropertyGenerationControlFlags::AllowProtected;
+
+	if (Property->HasAnyPropertyFlags(CPF_EditorOnly) || (!bAllowTransient && Property->HasAnyPropertyFlags(CPF_Transient)) )
 	{
 		UE_LOG(LogK2Compiler, Verbose, TEXT("FEmitDefaultValueHelper Skip EditorOnly or Transient property: %s"), *Property->GetPathName());
 		return;
@@ -128,7 +134,7 @@ void FEmitDefaultValueHelper::OuterGenerate(FEmitterLocalContext& Context
 
 			const uint8* ValuePtr = Property->ContainerPtrToValuePtr<uint8>(DataContainer, ArrayIndex);
 			const uint8* DefaultValuePtr = OptionalDefaultDataContainer ? Property->ContainerPtrToValuePtr<uint8>(OptionalDefaultDataContainer, ArrayIndex) : nullptr;
-			InnerGenerate(Context, Property, PathToMember, ValuePtr, DefaultValuePtr);
+			InnerGenerate(Context, Property, PathToMember, ValuePtr, DefaultValuePtr, ControlFlags | EPropertyGenerationControlFlags::IncludeFirstConstructionLine);
 		}
 	}
 }
@@ -162,8 +168,13 @@ void FEmitDefaultValueHelper::GenerateUserStructConstructor(const UUserDefinedSt
 	Context.Body.AddLine(TEXT("}"));
 }
 
-void FEmitDefaultValueHelper::InnerGenerate(FEmitterLocalContext& Context, const UProperty* Property, const FString& PathToMember, const uint8* ValuePtr, const uint8* DefaultValuePtr, bool bWithoutFirstConstructionLine)
+void FEmitDefaultValueHelper::InnerGenerate(FEmitterLocalContext& Context, const UProperty* Property, const FString& PathToMember, const uint8* ValuePtr, const uint8* DefaultValuePtr, EPropertyGenerationControlFlags ControlFlags)
 {
+	const bool bWithoutFirstConstructionLine = !(ControlFlags & EPropertyGenerationControlFlags::IncludeFirstConstructionLine);
+
+	// Don't propagate this flag through to nested calls.
+	ControlFlags &= ~EPropertyGenerationControlFlags::IncludeFirstConstructionLine;
+
 	auto InlineValueStruct = [&](UScriptStruct* OuterStruct, const uint8* LocalValuePtr) -> UScriptStruct*
 	{ 
 		UScriptStruct* InnerStruct = FBackendHelperUMG::InlineValueStruct(OuterStruct, LocalValuePtr);
@@ -246,7 +257,8 @@ void FEmitDefaultValueHelper::InnerGenerate(FEmitterLocalContext& Context, const
 		{
 			OuterGenerate(Context, LocalProperty, ActualPathToMember, ActualValuePtr
 				, (ActualDefaultValuePtr ? ActualDefaultValuePtr : DefaultStructOnScope.GetStructMemory())
-				, InnerInlineStruct ? EPropertyAccessOperator::Pointer : EPropertyAccessOperator::Dot);
+				, InnerInlineStruct ? EPropertyAccessOperator::Pointer : EPropertyAccessOperator::Dot
+				, ControlFlags);
 		}
 	}
 
@@ -275,7 +287,7 @@ void FEmitDefaultValueHelper::InnerGenerate(FEmitterLocalContext& Context, const
 		return bInitializeWithoutScriptStruct ? EStructConstructionType::EmptyConstructor : EStructConstructionType::Custom;
 	};
 
-	auto CreateElementSimple = [OneLineConstruction](FEmitterLocalContext& LocalContext, const UProperty* LocalProperty, const uint8* LocalValuePtr) -> FString
+	auto CreateElementSimple = [OneLineConstruction, ControlFlags](FEmitterLocalContext& LocalContext, const UProperty* LocalProperty, const uint8* LocalValuePtr) -> FString
 	{
 		FString ValueStr;
 		const bool bComplete = OneLineConstruction(LocalContext, LocalProperty, LocalValuePtr, ValueStr, true);
@@ -284,7 +296,7 @@ void FEmitDefaultValueHelper::InnerGenerate(FEmitterLocalContext& Context, const
 		{
 			const FString ElemLocName = LocalContext.GenerateUniqueLocalName();
 			LocalContext.AddLine(FString::Printf(TEXT("auto %s = %s;"), *ElemLocName, *ValueStr));
-			InnerGenerate(LocalContext, LocalProperty, ElemLocName, LocalValuePtr, nullptr, /*bWithoutFirstConstructionLine=*/true);
+			InnerGenerate(LocalContext, LocalProperty, ElemLocName, LocalValuePtr, nullptr, ControlFlags);
 			ValueStr = ElemLocName;
 		}
 		return ValueStr;
@@ -314,7 +326,7 @@ void FEmitDefaultValueHelper::InnerGenerate(FEmitterLocalContext& Context, const
 					const FString ArrayElementRefName = Context.GenerateUniqueLocalName();
 					Context.AddLine(FString::Printf(TEXT("auto& %s = %s[%d];"), *ArrayElementRefName, *PathToMember, Index));
 					// This is a Regular Struct (no special constructor), so we don't need to call constructor
-					InnerGenerate(Context, ArrayProperty->Inner, ArrayElementRefName, ScriptArrayHelper.GetRawPtr(Index), nullptr, /*bWithoutFirstConstructionLine=*/ true);
+					InnerGenerate(Context, ArrayProperty->Inner, ArrayElementRefName, ScriptArrayHelper.GetRawPtr(Index), nullptr, ControlFlags);
 				}
 			}
 			else
@@ -331,7 +343,7 @@ void FEmitDefaultValueHelper::InnerGenerate(FEmitterLocalContext& Context, const
 					if (!bComplete)
 					{
 						// The constructor was already called
-						InnerGenerate(Context, ArrayProperty->Inner, FString::Printf(TEXT("%s[%d]"), *PathToMember, Index), LocalValuePtr, nullptr, /*bWithoutFirstConstructionLine=*/ true);
+						InnerGenerate(Context, ArrayProperty->Inner, FString::Printf(TEXT("%s[%d]"), *PathToMember, Index), LocalValuePtr, nullptr, ControlFlags);
 					}
 				}
 			}
@@ -373,7 +385,7 @@ void FEmitDefaultValueHelper::InnerGenerate(FEmitterLocalContext& Context, const
 					const FString ElementName = Context.GenerateUniqueLocalName();
 					Context.AddLine(FString::Printf(TEXT("%s& %s = *(%s*)%s.GetElementPtr(%s.AddDefaultValue_Invalid_NeedsRehash());")
 						, *StructCppName, *ElementName, *StructCppName, *SetHelperName, *SetHelperName));
-					InnerGenerate(Context, StructProperty, ElementName, ScriptSetHelper.GetElementPtr(Index), nullptr, /*bWithoutFirstConstructionLine=*/true);
+					InnerGenerate(Context, StructProperty, ElementName, ScriptSetHelper.GetElementPtr(Index), nullptr, ControlFlags);
 				});
 				Context.AddLine(FString::Printf(TEXT("%s.Rehash();"), *SetHelperName));
 			}
@@ -441,7 +453,7 @@ void FEmitDefaultValueHelper::InnerGenerate(FEmitterLocalContext& Context, const
 						}
 						if (!bKeyComplete)
 						{
-							InnerGenerate(Context, MapProperty->KeyProp, KeyPath, ScriptMapHelper.GetKeyPtr(Index), nullptr, /*bWithoutFirstConstructionLine=*/true);
+							InnerGenerate(Context, MapProperty->KeyProp, KeyPath, ScriptMapHelper.GetKeyPtr(Index), nullptr, ControlFlags);
 						}
 					}
 
@@ -459,7 +471,7 @@ void FEmitDefaultValueHelper::InnerGenerate(FEmitterLocalContext& Context, const
 						}
 						if (!bValueComplete)
 						{
-							InnerGenerate(Context, MapProperty->ValueProp, ValuePath, ScriptMapHelper.GetValuePtr(Index), nullptr, /*bWithoutFirstConstructionLine=*/true);
+							InnerGenerate(Context, MapProperty->ValueProp, ValuePath, ScriptMapHelper.GetValuePtr(Index), nullptr, ControlFlags);
 						}
 					}
 
@@ -963,7 +975,7 @@ protected:
 
 			// Emit the C++ code needed to initialize the remainder of the struct's value within the component that's now being instanced as a default subobject within the converted context.
 			const FString PathToMember = FString::Printf(TEXT("%s->BodyInstance"), *VariableName);
-			FEmitDefaultValueHelper::InnerGenerate(Context, BodyInstanceProperty, PathToMember, (const uint8*)&Component->BodyInstance, BodyInstanceToCompare.GetStructMemory());
+			FEmitDefaultValueHelper::InnerGenerate(Context, BodyInstanceProperty, PathToMember, (const uint8*)&Component->BodyInstance, BodyInstanceToCompare.GetStructMemory(), FEmitDefaultValueHelper::EPropertyGenerationControlFlags::IncludeFirstConstructionLine);
 		}
 		else
 		{
@@ -1956,7 +1968,7 @@ void FEmitDefaultValueHelper::GenerateConstructor(FEmitterLocalContext& Context)
 					if(!FBackendHelperAnim::ShouldAddAnimNodeInitializationFunctionCall(Context, Property))
 					{
 						const bool bNewProperty = Property->GetOwnerStruct() == BPGC;
-						OuterGenerate(Context, Property, TEXT(""), reinterpret_cast<const uint8*>(CDO), bNewProperty ? nullptr : reinterpret_cast<const uint8*>(ParentCDO), EPropertyAccessOperator::None, true);
+						OuterGenerate(Context, Property, TEXT(""), reinterpret_cast<const uint8*>(CDO), bNewProperty ? nullptr : reinterpret_cast<const uint8*>(ParentCDO), EPropertyAccessOperator::None, EPropertyGenerationControlFlags::AllowProtected);
 					}
 				}
 			}
diff --git a/Engine/Source/Editor/KismetCompiler/Private/KismetCompiler.cpp b/Engine/Source/Editor/KismetCompiler/Private/KismetCompiler.cpp
--- a/Engine/Source/Editor/KismetCompiler/Private/KismetCompiler.cpp
+++ b/Engine/Source/Editor/KismetCompiler/Private/KismetCompiler.cpp
@@ -4230,6 +4230,9 @@ void FKismetCompilerContext::CompileFunctions(EInternalCompilerFlags InternalFla
 			}
 
 			PropagateValuesToCDO(NewCDO, OldCDO);
+
+			// Perform any fixup or caching based on the new CDO.
+			PostCDOCompiled();
 		}
 
 		// Note: The old->new CDO copy is deferred when regenerating, so we skip this step in that case.
diff --git a/Engine/Source/Runtime/Engine/Classes/Animation/AnimClassData.h b/Engine/Source/Runtime/Engine/Classes/Animation/AnimClassData.h
--- a/Engine/Source/Runtime/Engine/Classes/Animation/AnimClassData.h
+++ b/Engine/Source/Runtime/Engine/Classes/Animation/AnimClassData.h
@@ -71,11 +71,6 @@ public:
 	virtual const TArray<FName>& GetSyncGroupNames() const override { return SyncGroupNames; }
 	virtual int32 GetSyncGroupIndex(FName SyncGroupName) const override { return SyncGroupNames.IndexOfByKey(SyncGroupName); }
 	virtual const TArray<FExposedValueHandler>& GetExposedValueHandlers() const { return EvaluateGraphExposedInputs; }
-	
-	void InitGraphExposedInputs(UObject* ForObj)
-	{
-		FExposedValueHandler::Initialize(EvaluateGraphExposedInputs, ForObj);
-	}
 
 #if WITH_EDITOR
 	void CopyFrom(IAnimClassInterface* AnimClass)
diff --git a/Engine/Source/Runtime/Engine/Classes/Animation/AnimClassInterface.h b/Engine/Source/Runtime/Engine/Classes/Animation/AnimClassInterface.h
--- a/Engine/Source/Runtime/Engine/Classes/Animation/AnimClassInterface.h
+++ b/Engine/Source/Runtime/Engine/Classes/Animation/AnimClassInterface.h
@@ -64,15 +64,19 @@ struct FAnimBlueprintFunction
 	TArray<int32> InputPoseNodeIndices;
 
 	/** The property of the output node, patched up during link */
+	UPROPERTY(transient)
 	UStructProperty* OutputPoseNodeProperty;
 
 	/** The properties of the input nodes, patched up during link */
+	UPROPERTY(transient)
 	TArray<UStructProperty*> InputPoseNodeProperties;
 
 	/** The input properties themselves */
+	UPROPERTY(transient)
 	TArray<UProperty*> InputProperties;
 
 	/** Whether this function is actually implemented by this class - it could just be a stub */
+	UPROPERTY(transient)
 	bool bImplemented;
 };
 
