Support a proper fallback to construct any native 'noexport' struct-typed literal term in a nativized Blueprint class.

Change summary:
- Modified FEmitHelper::LiteralTerm() to emit UScriptStruct access through FindGloballyMappedObject() rather than explicitly emitting a call to StaticStruct().
- Modified FEmitDefaultValueHelper::SpecialStructureConstructor() to add explicit (fast path) construction for FSoftClassPath/FSoftObjectPath noexport types.
- Cleaned up some overuse of 'auto' in BlueprintCompilerCppBackendUtils.cpp.

diff --git a/Engine/Source/Developer/BlueprintCompilerCppBackend/Private/BlueprintCompilerCppBackendUtils.cpp b/Engine/Source/Developer/BlueprintCompilerCppBackend/Private/BlueprintCompilerCppBackendUtils.cpp
--- a/Engine/Source/Developer/BlueprintCompilerCppBackend/Private/BlueprintCompilerCppBackendUtils.cpp
+++ b/Engine/Source/Developer/BlueprintCompilerCppBackend/Private/BlueprintCompilerCppBackendUtils.cpp
@@ -137,8 +137,8 @@ FString FEmitterLocalContext::FindGloballyMappedObject(const UObject* Object, co
 	}
 
 	{
-		auto Field = Cast<UField>(Object);
-		auto FieldOwnerStruct = Field ? Field->GetOwnerStruct() : nullptr;
+		const UField* Field = Cast<UField>(Object);
+		const UStruct* FieldOwnerStruct = Field ? Field->GetOwnerStruct() : nullptr;
 		if (FieldOwnerStruct && (Field != FieldOwnerStruct))
 		{
 			ensure(Field == FindField<UField>(FieldOwnerStruct, Field->GetFName()));
@@ -160,18 +160,16 @@ FString FEmitterLocalContext::FindGloballyMappedObject(const UObject* Object, co
 		}
 	}
 
-	if (auto ObjClass = Cast<UClass>(Object))
+	if (const UClass* ObjClass = Cast<UClass>(Object))
 	{
-		auto BPGC = Cast<UBlueprintGeneratedClass>(ObjClass);
+		const UBlueprintGeneratedClass* BPGC = Cast<UBlueprintGeneratedClass>(ObjClass);
 		if (ObjClass->HasAnyClassFlags(CLASS_Native) || (BPGC && Dependencies.WillClassBeConverted(BPGC)))
 		{
 			return CastCustomClass(FString::Printf(TEXT("%s::StaticClass()"), *FEmitHelper::GetCppName(ObjClass, true)));
 		}
 	}
 
-	// TODO Handle native structires, and special cases..
-
-	if (auto ScriptStruct = Cast<UScriptStruct>(Object))
+	if (const UScriptStruct* ScriptStruct = Cast<UScriptStruct>(Object))
 	{
 		if (ScriptStruct->StructFlags & STRUCT_NoExport)
 		{
@@ -179,8 +177,6 @@ FString FEmitterLocalContext::FindGloballyMappedObject(const UObject* Object, co
 		}
 		else
 		{
-			// Check if  
-			// TODO: check if supported 
 			return FString::Printf(TEXT("%s::StaticStruct()"), *FEmitHelper::GetCppName(ScriptStruct));
 		}
 	}
@@ -279,7 +275,7 @@ FString FEmitterLocalContext::ExportTextItem(const UProperty* Property, const vo
 		| EPropertyExportCPPFlags::CPPF_NoRef
 		| EPropertyExportCPPFlags::CPPF_NoStaticArray
 		| EPropertyExportCPPFlags::CPPF_BlueprintCppBackend;
-	if (auto ArrayProperty = Cast<const UArrayProperty>(Property))
+	if (const UArrayProperty* ArrayProperty = Cast<const UArrayProperty>(Property))
 	{
 		const FString ConstPrefix = Property->HasMetaData(TEXT("NativeConstTemplateArg")) ? TEXT("const ") : TEXT("");
 		const FString TypeText = ExportCppDeclaration(ArrayProperty, EExportedDeclaration::Parameter, LocalExportCPPFlags, EPropertyNameInDeclaration::Skip, FString(), ConstPrefix);
@@ -314,7 +310,7 @@ FString FEmitterLocalContext::ExportCppDeclaration(const UProperty* Property, EE
 		{
 			FString Result;
 
-			auto BPGC = Cast<UBlueprintGeneratedClass>(InActualClass);
+			UBlueprintGeneratedClass* BPGC = Cast<UBlueprintGeneratedClass>(InActualClass);
 			if (BPGC || !TypePrefix.IsEmpty())
 			{
 				UClass* NativeType = GetFirstNativeOrConvertedClass(InActualClass);
@@ -332,19 +328,19 @@ FString FEmitterLocalContext::ExportCppDeclaration(const UProperty* Property, EE
 		{
 			Result = GetCppTypeFromObjectProperty(ClassProperty, ClassProperty->MetaClass);
 		}
-		else if (auto SoftClassProperty = Cast<const USoftClassProperty>(InProperty))
+		else if (const USoftClassProperty* SoftClassProperty = Cast<const USoftClassProperty>(InProperty))
 		{
 			Result = GetCppTypeFromObjectProperty(SoftClassProperty, SoftClassProperty->MetaClass);
 		}
-		else if (auto ObjectProperty = Cast<const UObjectPropertyBase>(InProperty))
+		else if (const UObjectPropertyBase* ObjectProperty = Cast<const UObjectPropertyBase>(InProperty))
 		{
 			Result = GetCppTypeFromObjectProperty(ObjectProperty, ObjectProperty->PropertyClass);
 		}
-		else if (auto StructProperty = Cast<const UStructProperty>(InProperty))
+		else if (const UStructProperty* StructProperty = Cast<const UStructProperty>(InProperty))
 		{
 			Result = FEmitHelper::GetCppName(StructProperty->Struct, false);
 		}
-		else if (auto SCDelegateProperty = Cast<const UDelegateProperty>(InProperty))
+		else if (const UDelegateProperty* SCDelegateProperty = Cast<const UDelegateProperty>(InProperty))
 		{
 			const FString* SCDelegateTypeName = MCDelegateSignatureToSCDelegateType.Find(SCDelegateProperty->SignatureFunction);
 			if (SCDelegateTypeName)
@@ -433,8 +429,8 @@ void FEmitterLocalContext::MarkUnconvertedClassAsNecessary(UField* InField)
 FString FEmitHelper::GetCppName(const UField* Field, bool bUInterface, bool bForceParameterNameModification)
 {
 	check(Field);
-	auto AsClass = Cast<UClass>(Field);
-	auto AsScriptStruct = Cast<UScriptStruct>(Field);
+	const UClass* AsClass = Cast<UClass>(Field);
+	const UScriptStruct* AsScriptStruct = Cast<UScriptStruct>(Field);
 	if (AsClass || AsScriptStruct)
 	{
 		if (AsClass && AsClass->HasAnyClassFlags(CLASS_Interface))
@@ -444,7 +440,7 @@ FString FEmitHelper::GetCppName(const UField* Field, bool bUInterface, bool bFor
 				, bUInterface ? TEXT("U") : TEXT("I")
 				, *AsClass->GetName());
 		}
-		auto AsStruct = CastChecked<UStruct>(Field);
+		const UStruct* AsStruct = CastChecked<UStruct>(Field);
 		if (AsStruct->IsNative())
 		{
 			return FString::Printf(TEXT("%s%s"), AsStruct->GetPrefixCPP(), *AsStruct->GetName());
@@ -454,7 +450,7 @@ FString FEmitHelper::GetCppName(const UField* Field, bool bUInterface, bool bFor
 			return ::UnicodeToCPPIdentifier(*AsStruct->GetName(), false, AsStruct->GetPrefixCPP()) + GetPathPostfix(AsStruct);
 		}
 	}
-	else if (auto AsProperty = Cast<UProperty>(Field))
+	else if (const UProperty* AsProperty = Cast<UProperty>(Field))
 	{
 		const UStruct* Owner = AsProperty->GetOwnerStruct();
 		const bool bModifyName = ensure(Owner)
@@ -797,17 +793,17 @@ FString FEmitHelper::GenerateReplaceConvertedMD(UObject* Obj)
 			{
 				static UObject* FindFinalObject(UObjectRedirector* Redirector)
 				{
-					auto DestObj = Redirector ? Redirector->DestinationObject : nullptr;
-					auto InnerRedir = Cast<UObjectRedirector>(DestObj);
+					UObject* DestObj = Redirector ? Redirector->DestinationObject : nullptr;
+					UObjectRedirector* InnerRedir = Cast<UObjectRedirector>(DestObj);
 					return InnerRedir ? FindFinalObject(InnerRedir) : DestObj;
 				}
 			};
 
 			TArray<UObject*> AllObjects;
 			GetObjectsOfClass(UObjectRedirector::StaticClass(), AllObjects);
-			for (auto LocalObj : AllObjects)
+			for (UObject* LocalObj : AllObjects)
 			{
-				auto Redirector = CastChecked<UObjectRedirector>(LocalObj);
+				UObjectRedirector* Redirector = CastChecked<UObjectRedirector>(LocalObj);
 				if (Obj == FLocalHelper::FindFinalObject(Redirector))
 				{
 					Result += TEXT(",");
@@ -842,7 +838,7 @@ FString FEmitHelper::GetBaseFilename(const UObject* AssetObj, const FCompilerNat
 	// based on the filename. I'm also not interested in exploring the depth of unicode filename support in UHT,
 	// UBT, and our various c++ toolchains, so this logic is pretty aggressive:
 	FString Postfix = TEXT("__pf");
-	for (auto& Char : AssetName)
+	for (TCHAR& Char : AssetName)
 	{
 		if (!IsValidCPPIdentifierChar(Char))
 		{
@@ -861,7 +857,7 @@ FString FEmitHelper::GetPCHFilename()
 	FString PCHFilename;
 	IBlueprintCompilerCppBackendModule& BackEndModule = (IBlueprintCompilerCppBackendModule&)IBlueprintCompilerCppBackendModule::Get();
 
-	auto& PchFilenameQuery = BackEndModule.OnPCHFilenameQuery();
+	TBaseDelegate<FString>& PchFilenameQuery = BackEndModule.OnPCHFilenameQuery();
 	if (PchFilenameQuery.IsBound())
 	{
 		PCHFilename = PchFilenameQuery.Execute();
@@ -953,7 +949,7 @@ void FEmitHelper::EmitMulticastDelegateDeclarations(FEmitterLocalContext& Emitte
 {
 	for (TFieldIterator<UMulticastDelegateProperty> It(EmitterContext.GetCurrentlyGeneratedClass(), EFieldIteratorFlags::ExcludeSuper); It; ++It)
 	{
-		auto Signature = It->SignatureFunction;
+		UFunction* Signature = It->SignatureFunction;
 		check(Signature);
 
 		FString ParamNumberStr, Parameters;
@@ -972,8 +968,8 @@ void FEmitHelper::EmitMulticastDelegateDeclarations(FEmitterLocalContext& Emitte
 
 void FEmitHelper::EmitLifetimeReplicatedPropsImpl(FEmitterLocalContext& EmitterContext)
 {
-	auto SourceClass = EmitterContext.GetCurrentlyGeneratedClass();
-	auto CppClassName = FEmitHelper::GetCppName(SourceClass);
+	UClass* SourceClass = EmitterContext.GetCurrentlyGeneratedClass();
+	FString CppClassName = FEmitHelper::GetCppName(SourceClass);
 	bool bFunctionInitilzed = false;
 	for (TFieldIterator<UProperty> It(SourceClass, EFieldIteratorFlags::ExcludeSuper); It; ++It)
 	{
@@ -1023,7 +1019,7 @@ FString FEmitHelper::FloatToString(float Value)
 
 FString FEmitHelper::LiteralTerm(FEmitterLocalContext& EmitterContext, const FLiteralTermParams& Params)
 {
-	auto Schema = GetDefault<UEdGraphSchema_K2>();
+	const UEdGraphSchema_K2* Schema = GetDefault<UEdGraphSchema_K2>();
 
 	class FImportTextErrorContext : public FStringOutputDevice
 	{
@@ -1221,7 +1217,7 @@ FString FEmitHelper::LiteralTerm(FEmitterLocalContext& EmitterContext, const FLi
 	}
 	else if ((UEdGraphSchema_K2::PC_Byte == Type.PinCategory) || (UEdGraphSchema_K2::PC_Enum == Type.PinCategory))
 	{
-		auto TypeEnum = Cast<UEnum>(Type.PinSubCategoryObject.Get());
+		UEnum* TypeEnum = Cast<UEnum>(Type.PinSubCategoryObject.Get());
 		if (TypeEnum)
 		{
 			// @note: We have to default to the zeroth entry because there may not be a symbol associated with the 
@@ -1336,9 +1332,13 @@ FString FEmitHelper::LiteralTerm(FEmitterLocalContext& EmitterContext, const FLi
 				}
 
 				{
+					// FindGloballyMappedObject() will re-route to FStructAccessHelper::EmitStructAccessCode() for 'noexport' types, and will fall back to <type>::StaticStruct() for other native cases.
+					const FString StructObjectVar = EmitterContext.GenerateUniqueLocalName();
+					EmitterContext.AddLine(FString::Printf(TEXT("const UScriptStruct* %s = %s;"), *StructObjectVar, *EmitterContext.FindGloballyMappedObject(StructType, UScriptStruct::StaticClass())));
+					
 					const FString StructMemoryVar = EmitterContext.GenerateUniqueLocalName();
-					EmitterContext.AddLine(FString::Printf(TEXT("uint8* %s = (uint8*)FMemory_Alloca(%s::StaticStruct()->GetStructureSize());"), *StructMemoryVar, *StructName));
-					EmitterContext.AddLine(FString::Printf(TEXT("%s::StaticStruct()->InitializeStruct(%s);"), *StructName, *StructMemoryVar));
+					EmitterContext.AddLine(FString::Printf(TEXT("uint8* %s = (uint8*)FMemory_Alloca(%s->GetStructureSize());"), *StructMemoryVar, *StructObjectVar));
+					EmitterContext.AddLine(FString::Printf(TEXT("%s->InitializeStruct(%s);"), *StructObjectVar, *StructMemoryVar));
 					EmitterContext.AddLine(FString::Printf(TEXT("%s& %s = *reinterpret_cast<%s*>(%s);"), *StructName, *LocalStructNativeName, *StructName, *StructMemoryVar));  // TODO: ?? should "::GetDefaultValue()" be called here?
 				}
 
@@ -1359,7 +1359,7 @@ FString FEmitHelper::LiteralTerm(FEmitterLocalContext& EmitterContext, const FLi
 	}
 	else if (UEdGraphSchema_K2::PC_Class == Type.PinCategory)
 	{
-		if (auto FoundClass = Cast<const UClass>(Params.LiteralObject))
+		if (const UClass* FoundClass = Cast<const UClass>(Params.LiteralObject))
 		{
 			const FString MappedObject = EmitterContext.FindGloballyMappedObject(Params.LiteralObject, UClass::StaticClass());
 			if (!MappedObject.IsEmpty())
@@ -1418,7 +1418,7 @@ FString FEmitHelper::PinTypeToNativeType(const FEdGraphPinType& Type)
 	// A temp uproperty should be generated?
 	auto PinTypeToNativeTypeInner = [](const FEdGraphPinType& InType) -> FString
 	{
-		auto Schema = GetDefault<UEdGraphSchema_K2>();
+		const UEdGraphSchema_K2* Schema = GetDefault<UEdGraphSchema_K2>();
 		if (UEdGraphSchema_K2::PC_String == InType.PinCategory)
 		{
 			return TEXT("FString");
@@ -1536,7 +1536,7 @@ UFunction* FEmitHelper::GetOriginalFunction(UFunction* Function)
 bool FEmitHelper::ShouldHandleAsNativeEvent(UFunction* Function, bool bOnlyIfOverridden)
 {
 	check(Function);
-	auto OriginalFunction = FEmitHelper::GetOriginalFunction(Function);
+	UFunction* OriginalFunction = FEmitHelper::GetOriginalFunction(Function);
 	check(OriginalFunction);
 	if (!bOnlyIfOverridden || (OriginalFunction != Function))
 	{
@@ -1550,7 +1550,7 @@ bool FEmitHelper::ShouldHandleAsNativeEvent(UFunction* Function, bool bOnlyIfOve
 bool FEmitHelper::ShouldHandleAsImplementableEvent(UFunction* Function)
 {
 	check(Function);
-	auto OriginalFunction = FEmitHelper::GetOriginalFunction(Function);
+	UFunction* OriginalFunction = FEmitHelper::GetOriginalFunction(Function);
 	check(OriginalFunction);
 	if (OriginalFunction != Function)
 	{
@@ -2070,7 +2070,7 @@ void FBackendHelperStaticSearchableValues::EmitFunctionDeclaration(FEmitterLocal
 
 void FBackendHelperStaticSearchableValues::EmitFunctionDefinition(FEmitterLocalContext& Context)
 {
-	auto BPGC = CastChecked<UBlueprintGeneratedClass>(Context.GetCurrentlyGeneratedClass());
+	UBlueprintGeneratedClass* BPGC = CastChecked<UBlueprintGeneratedClass>(Context.GetCurrentlyGeneratedClass());
 	const FString CppClassName = FEmitHelper::GetCppName(BPGC);
 	const FString FunctionName = GetFunctionName();
 
@@ -2248,8 +2248,8 @@ void FScopeBlock::TrackLocalAccessorDecl(const UProperty* Property)
 
 struct FStructAccessHelper_StaticData
 {
-	TMap<const UStruct*, FString> BaseStructureAccessorsMap;
-	TMap<const UStruct*, bool> SupportsDirectNativeAccessMap;
+	TMap<const UScriptStruct*, FString> BaseStructureAccessorsMap;
+	TMap<const UScriptStruct*, bool> SupportsDirectNativeAccessMap;
 	TArray<FSoftClassPath> NoExportTypesWithDirectNativeFieldAccess;
 
 	static FStructAccessHelper_StaticData& Get()
@@ -2292,7 +2292,7 @@ private:
 	}
 };
 
-FString FStructAccessHelper::EmitStructAccessCode(const UStruct* InStruct)
+FString FStructAccessHelper::EmitStructAccessCode(const UScriptStruct* InStruct)
 {
 	check(InStruct != nullptr);
 	
@@ -2302,7 +2302,7 @@ FString FStructAccessHelper::EmitStructAccessCode(const UStruct* InStruct)
 	}
 	else
 	{
-		return FString::Printf(TEXT("FStructUtils::FindStructureInPackageChecked(TEXT(\"%s\"), TEXT(\"%s\"))"), *InStruct->GetName(), *InStruct->GetOutermost()->GetName());
+		return FString::Printf(TEXT("CastChecked<UScriptStruct>(FStructUtils::FindStructureInPackageChecked(TEXT(\"%s\"), TEXT(\"%s\")))"), *InStruct->GetName(), *InStruct->GetOutermost()->GetName());
 	}
 }
 
diff --git a/Engine/Source/Developer/BlueprintCompilerCppBackend/Private/BlueprintCompilerCppBackendUtils.h b/Engine/Source/Developer/BlueprintCompilerCppBackend/Private/BlueprintCompilerCppBackendUtils.h
--- a/Engine/Source/Developer/BlueprintCompilerCppBackend/Private/BlueprintCompilerCppBackendUtils.h
+++ b/Engine/Source/Developer/BlueprintCompilerCppBackend/Private/BlueprintCompilerCppBackendUtils.h
@@ -464,6 +464,6 @@ public:
 
 struct FStructAccessHelper
 {
-	static FString EmitStructAccessCode(const UStruct* InStruct);
+	static FString EmitStructAccessCode(const UScriptStruct* InStruct);
 	static bool CanEmitDirectFieldAccess(const UScriptStruct* InStruct);
 };
\ No newline at end of file
diff --git a/Engine/Source/Developer/BlueprintCompilerCppBackend/Private/BlueprintCompilerCppBackendValueHelper.cpp b/Engine/Source/Developer/BlueprintCompilerCppBackend/Private/BlueprintCompilerCppBackendValueHelper.cpp
--- a/Engine/Source/Developer/BlueprintCompilerCppBackend/Private/BlueprintCompilerCppBackendValueHelper.cpp
+++ b/Engine/Source/Developer/BlueprintCompilerCppBackend/Private/BlueprintCompilerCppBackendValueHelper.cpp
@@ -690,6 +690,26 @@ bool FEmitDefaultValueHelper::SpecialStructureConstructor(const UStruct* Struct,
 		return true;
 	}
 
+	if (TBaseStructure<FSoftClassPath>::Get() == Struct)
+	{
+		if (OutResult)
+		{
+			const FSoftClassPath* SoftClassPath = reinterpret_cast<const FSoftClassPath*>(ValuePtr);
+			*OutResult = FString::Printf(TEXT("FSoftClassPath(TEXT(\"%s\"))"), *SoftClassPath->ToString());
+		}
+		return true;
+	}
+
+	if (TBaseStructure<FSoftObjectPath>::Get() == Struct)
+	{
+		if (OutResult)
+		{
+			const FSoftObjectPath* SoftObjectPath = reinterpret_cast<const FSoftObjectPath*>(ValuePtr);
+			*OutResult = FString::Printf(TEXT("FSoftObjectPath(TEXT(\"%s\"))"), *SoftObjectPath->ToString());
+		}
+		return true;
+	}
+
 	return false;
 }
 
