Fix for an EDL runtime assertion on load when constructing the CDO for a nativized Blueprint class that contains an instanced object reference of another nativized Blueprint class type that references an instance of a non-nativized Blueprint class as a nested, class-owned subobject.

Change summary:
- Modified FGatherConvertedClassDependenciesHelperBase::AddConvertedClassDependency() to add converted CDOs as a dependency in addition to the class itself. This ensures (through EDL) that any nested, non-nativized BPGC dependencies will get serialized prior to construction of the converted CDO.

diff --git a/Engine/Source/Developer/BlueprintCompilerCppBackend/Private/BlueprintCompilerCppBackendGatherDependencies.cpp b/Engine/Source/Developer/BlueprintCompilerCppBackend/Private/BlueprintCompilerCppBackendGatherDependencies.cpp
--- a/Engine/Source/Developer/BlueprintCompilerCppBackend/Private/BlueprintCompilerCppBackendGatherDependencies.cpp
+++ b/Engine/Source/Developer/BlueprintCompilerCppBackend/Private/BlueprintCompilerCppBackendGatherDependencies.cpp
@@ -77,6 +77,9 @@ struct FGatherConvertedClassDependenciesHelperBase : public FReferenceCollector
 		if (InBPGC && !Dependencies.ConvertedClasses.Contains(InBPGC))
 		{
 			Dependencies.ConvertedClasses.Add(InBPGC);
+
+			// Add the CDO in case we need to instance any subobjects at construction time, which makes the converted CDO a dependency.
+			Dependencies.Assets.AddUnique(InBPGC->GetDefaultObject());
 		}
 	}
 
