Remove broken test (as best it can be told it could never have worked)
#rnx

diff --git a/Engine/Source/Developer/BlueprintCompilerCppBackend/Private/BPCompilerTests.cpp b/Engine/Source/Developer/BlueprintCompilerCppBackend/Private/BPCompilerTests.cpp
--- a/Engine/Source/Developer/BlueprintCompilerCppBackend/Private/BPCompilerTests.cpp
+++ b/Engine/Source/Developer/BlueprintCompilerCppBackend/Private/BPCompilerTests.cpp
@@ -464,28 +464,4 @@ bool FBPCompilerNodeTest::RunTest(const FString& Parameters)
 	};
 
 	return RunTestHelper(TestBody, this);
-}
-
-IMPLEMENT_SIMPLE_AUTOMATION_TEST(FBPCompilerLatentTest, "Project.Blueprints.NativeBackend.Latent", CompilerTestFlags)
-bool FBPCompilerLatentTest::RunTest(const FString& Parameters)
-{
-	auto TestBody = [](ClassAccessor F, FAutomationTestBase* Context)
-	{
-		FOwnedObjectsHelper OwnedObjects;
-		TGuardValue<bool> AutoRestore(GAllowActorScriptExecutionInEditor, true);
-
-		UObject* TestInstance = NewTestActor(F(TEXT("Node"), TEXT("BP_Latent_Basic"), Context, OwnedObjects), OwnedObjects);
-		if (!TestInstance)
-		{
-			return 0u;
-		}
-
-		Call(TestInstance, TEXT("RunDelayTest"));
-		Call(TestInstance, TEXT("RunDownloadTest"));
-
-		FArchiveSkipTransientObjectCRC32 Results;
-		return Results.Crc32(TestInstance);
-	};
-
-	return RunTestHelper(TestBody, this);
-}
+}
\ No newline at end of file
