Fix a Blueprint nativization compile failure due to one or more unconnected TSoftClassPtr input pins on an event or function call node.

Change summary:
- Modified FEmitHelper::LiteralTerm() to generalize TSoftClassPtr/TSoftObjectPtr terms - the default value case did not previously consider TSoftClassPtr pin types.

diff --git a/Engine/Source/Developer/BlueprintCompilerCppBackend/Private/BlueprintCompilerCppBackendUtils.cpp b/Engine/Source/Developer/BlueprintCompilerCppBackend/Private/BlueprintCompilerCppBackendUtils.cpp
--- a/Engine/Source/Developer/BlueprintCompilerCppBackend/Private/BlueprintCompilerCppBackendUtils.cpp
+++ b/Engine/Source/Developer/BlueprintCompilerCppBackend/Private/BlueprintCompilerCppBackendUtils.cpp
@@ -1376,15 +1376,16 @@ FString FEmitHelper::LiteralTerm(FEmitterLocalContext& EmitterContext, const FLi
 		MetaClass = MetaClass ? MetaClass : UObject::StaticClass();
 		const FString ObjTypeStr = FEmitHelper::GetCppName(EmitterContext.GetFirstNativeOrConvertedClass(MetaClass));
 
+		const bool bAssetSubclassOf = (UEdGraphSchema_K2::PC_SoftClass == Type.PinCategory);
+		const FString TermTypeStr = bAssetSubclassOf ? TEXT("TSoftClassPtr") : TEXT("TSoftObjectPtr");
+
+		FString TermValueStr;
 		if (!CustomValue.IsEmpty())
 		{
-			const bool bAssetSubclassOf = (UEdGraphSchema_K2::PC_SoftClass == Type.PinCategory);
-			return FString::Printf(TEXT("%s<%s>(FSoftObjectPath(TEXT(\"%s\")))")
-				, bAssetSubclassOf ? TEXT("TSoftClassPtr") : TEXT("TSoftObjectPtr")
-				, *ObjTypeStr
-				, *(CustomValue.ReplaceCharWithEscapedChar()));
+			TermValueStr = FString::Printf(TEXT("FSoftObjectPath(TEXT(\"%s\"))"), *(CustomValue.ReplaceCharWithEscapedChar()));
 		}
-		return FString::Printf(TEXT("((%s*)nullptr)"), *ObjTypeStr);
+
+		return FString::Printf(TEXT("%s<%s>(%s)"), *TermTypeStr, *ObjTypeStr, *TermValueStr);
 	}
 	else if (UEdGraphSchema_K2::PC_Object == Type.PinCategory)
 	{
