Fix uninitialized user-defined struct fields in nativized, cooked builds.

Change summary:
- Modified FEmitDefaultValueHelper::GenerateUserStructConstructor() to ensure that POD fields are always initialized in the converted UDS ctor.

diff --git a/Engine/Source/Developer/BlueprintCompilerCppBackend/Private/BlueprintCompilerCppBackendValueHelper.cpp b/Engine/Source/Developer/BlueprintCompilerCppBackend/Private/BlueprintCompilerCppBackendValueHelper.cpp
--- a/Engine/Source/Developer/BlueprintCompilerCppBackend/Private/BlueprintCompilerCppBackendValueHelper.cpp
+++ b/Engine/Source/Developer/BlueprintCompilerCppBackend/Private/BlueprintCompilerCppBackendValueHelper.cpp
@@ -152,7 +152,9 @@ void FEmitDefaultValueHelper::GenerateUserStructConstructor(const UUserDefinedSt
 		FUserStructOnScopeIgnoreDefaults RawDefaultStructOnScope(Struct);
 		for (auto Property : TFieldRange<const UProperty>(Struct))
 		{
-			OuterGenerate(Context, Property, TEXT(""), StructData.GetStructMemory(), RawDefaultStructOnScope.GetStructMemory(), EPropertyAccessOperator::None);
+			// Since UDS types are converted to native USTRUCT, all POD fields must be initialized in the ctor, just as with "regular" native USTRUCT types.
+			const bool bForceInit = Property->HasAnyPropertyFlags(CPF_IsPlainOldData);
+			OuterGenerate(Context, Property, TEXT(""), StructData.GetStructMemory(), bForceInit ? nullptr : RawDefaultStructOnScope.GetStructMemory(), EPropertyAccessOperator::None);
 		}
 	}
 	Context.Body.DecreaseIndent();
