Fix a regression that causes nativized QAGame to fail when loading test maps containing nativized BP actor instances with circular references.

Notes:
- This will re-introduce UE-62514. I will be reopening that issue to be addressed for the next release instead.

#rnx

diff --git a/Engine/Source/Developer/BlueprintCompilerCppBackend/Private/BlueprintCompilerCppBackendGatherDependencies.cpp b/Engine/Source/Developer/BlueprintCompilerCppBackend/Private/BlueprintCompilerCppBackendGatherDependencies.cpp
--- a/Engine/Source/Developer/BlueprintCompilerCppBackend/Private/BlueprintCompilerCppBackendGatherDependencies.cpp
+++ b/Engine/Source/Developer/BlueprintCompilerCppBackend/Private/BlueprintCompilerCppBackendGatherDependencies.cpp
@@ -77,9 +77,6 @@ struct FGatherConvertedClassDependenciesHelperBase : public FReferenceCollector
 		if (InBPGC && !Dependencies.ConvertedClasses.Contains(InBPGC))
 		{
 			Dependencies.ConvertedClasses.Add(InBPGC);
-
-			// Add the CDO in case we need to instance any subobjects at construction time, which makes the converted CDO a dependency.
-			Dependencies.Assets.AddUnique(InBPGC->GetDefaultObject());
 		}
 	}
 
