Allow nativized Blueprint types to instance other nativized Blueprint types as subobjects when they have one or more non-Blueprint asset preload dependencies (e.g. non-native UMG widgets).

Change summary:
- Added a shared cache to FGatherConvertedClassDependencies. The ctor is now protected and access is done through the static Get() method. This prevents redundant dependency list gathering for a previously-queried class/struct type.
- Switched FGatherConvertedClassDependencies::Assets from a TArray to a TSet (for consistency w/ the other dependency lists).
- Added FGatherConvertedClassDependenciesHelperBase::AddAssetDependency() to consolidate the addition of non-converted type asset dependencies (for consistency and easier maintenance).
- Added FGatherConvertedClassDependencies::GatherAssetsReferencedByConvertedTypes() as a utility method for building a nested asset dependency chain when converted subobject types are involved.
- Modified FFakeImportTableHelper::GetClassesOfSubobjects() to flag nested asset dependencies as CDO preload dependencies for EDL on the outermost converted type. This ensures dependencies will be loaded prior to CDO construction before instancing subobjects of a converted class type.
- Modified FGatherConvertedClassDependencies::AllDependencies() to include nested asset dependencies; this ensures they're included when we're generating the code needed to support the linker's "fake" export table at runtime.
- Some misc. cleanup; replaced a few instances of 'auto' with explicit types and updated a few lambda capture lists in place of previous use of '&' as a catch-all.

diff --git a/Engine/Source/Developer/BlueprintCompilerCppBackend/Private/BlueprintCompilerCppBackendBase.cpp b/Engine/Source/Developer/BlueprintCompilerCppBackend/Private/BlueprintCompilerCppBackendBase.cpp
--- a/Engine/Source/Developer/BlueprintCompilerCppBackend/Private/BlueprintCompilerCppBackendBase.cpp
+++ b/Engine/Source/Developer/BlueprintCompilerCppBackend/Private/BlueprintCompilerCppBackendBase.cpp
@@ -323,11 +323,11 @@ FString FBlueprintCompilerCppBackendBase::GenerateCodeFromClass(UClass* SourceCl
 	auto CleanCppClassName = FEmitHelper::GetBaseFilename(SourceClass, NativizationOptions);
 	auto CppClassName = FEmitHelper::GetCppName(SourceClass);
 	
-	FGatherConvertedClassDependencies Dependencies(SourceClass, NativizationOptions);
-	FNativizationSummaryHelper::RegisterRequiredModules(NativizationOptions.PlatformName, Dependencies.RequiredModuleNames);
-	FEmitterLocalContext EmitterContext(Dependencies, NativizationOptions);
+	TSharedPtr<FGatherConvertedClassDependencies> Dependencies = FGatherConvertedClassDependencies::Get(SourceClass, NativizationOptions);
+	FNativizationSummaryHelper::RegisterRequiredModules(NativizationOptions.PlatformName, Dependencies->RequiredModuleNames);
+	FEmitterLocalContext EmitterContext(Dependencies.ToSharedRef().Get(), NativizationOptions);
 
-	UClass* OriginalSourceClass = Dependencies.FindOriginalClass(SourceClass);
+	UClass* OriginalSourceClass = Dependencies->FindOriginalClass(SourceClass);
 	ensure(OriginalSourceClass != SourceClass);
 
 	FNativizationSummaryHelper::RegisterClass(OriginalSourceClass);
@@ -427,9 +427,7 @@ FString FBlueprintCompilerCppBackendBase::GenerateCodeFromClass(UClass* SourceCl
 		{
 			UBlueprintGeneratedClass* BPGC = CastChecked<UBlueprintGeneratedClass>(EmitterContext.GetCurrentlyGeneratedClass());
 			UBlueprintGeneratedClass* ParentBPGC = Cast<UBlueprintGeneratedClass>(BPGC->GetSuperClass());
-			ParentDependencies = TSharedPtr<FGatherConvertedClassDependencies>( ParentBPGC 
-				? new FGatherConvertedClassDependencies(ParentBPGC, NativizationOptions) 
-				: nullptr);
+			ParentDependencies = FGatherConvertedClassDependencies::Get(ParentBPGC, NativizationOptions);
 
 			EmitterContext.Header.AddLine(FString::Printf(TEXT("%s(const FObjectInitializer& ObjectInitializer = FObjectInitializer::Get());"), *CppClassName));
 			EmitterContext.Header.AddLine(TEXT("virtual void PostLoadSubobjects(FObjectInstancingGraph* OuterInstanceGraph) override;"));
@@ -1041,9 +1039,9 @@ void FBlueprintCompilerCppBackendBase::GenerateCodeFromEnum(UUserDefinedEnum* So
 void FBlueprintCompilerCppBackendBase::GenerateCodeFromStruct(UUserDefinedStruct* SourceStruct, const FCompilerNativizationOptions& NativizationOptions, FString& OutHeaderCode, FString& OutCPPCode)
 {
 	check(SourceStruct);
-	FGatherConvertedClassDependencies Dependencies(SourceStruct, NativizationOptions);
-	FNativizationSummaryHelper::RegisterRequiredModules(NativizationOptions.PlatformName, Dependencies.RequiredModuleNames);
-	FEmitterLocalContext EmitterContext(Dependencies, NativizationOptions);
+	TSharedPtr<FGatherConvertedClassDependencies> Dependencies = FGatherConvertedClassDependencies::Get(SourceStruct, NativizationOptions);
+	FNativizationSummaryHelper::RegisterRequiredModules(NativizationOptions.PlatformName, Dependencies->RequiredModuleNames);
+	FEmitterLocalContext EmitterContext(Dependencies.ToSharedRef().Get(), NativizationOptions);
 	// use GetBaseFilename() so that we can coordinate #includes and filenames
 	EmitFileBeginning(FEmitHelper::GetBaseFilename(SourceStruct, NativizationOptions), EmitterContext, true, true);
 	{
@@ -1088,9 +1086,9 @@ void FBlueprintCompilerCppBackendBase::GenerateCodeFromStruct(UUserDefinedStruct
 
 FString FBlueprintCompilerCppBackendBase::GenerateWrapperForClass(UClass* SourceClass, const FCompilerNativizationOptions& NativizationOptions)
 {
-	FGatherConvertedClassDependencies Dependencies(SourceClass, NativizationOptions);
-	FNativizationSummaryHelper::RegisterRequiredModules(NativizationOptions.PlatformName, Dependencies.RequiredModuleNames);
-	FEmitterLocalContext EmitterContext(Dependencies, NativizationOptions);
+	TSharedPtr<FGatherConvertedClassDependencies> Dependencies = FGatherConvertedClassDependencies::Get(SourceClass, NativizationOptions);
+	FNativizationSummaryHelper::RegisterRequiredModules(NativizationOptions.PlatformName, Dependencies->RequiredModuleNames);
+	FEmitterLocalContext EmitterContext(Dependencies.ToSharedRef().Get(), NativizationOptions);
 
 	UBlueprintGeneratedClass* BPGC = Cast<UBlueprintGeneratedClass>(SourceClass);
 
@@ -1140,7 +1138,7 @@ FString FBlueprintCompilerCppBackendBase::GenerateWrapperForClass(UClass* Source
 					break;
 				}
 				UBlueprintGeneratedClass* SuperBPGC = Cast<UBlueprintGeneratedClass>(SuperClassToUse);
-				if (SuperBPGC && Dependencies.WillClassBeConverted(SuperBPGC))
+				if (SuperBPGC && Dependencies->WillClassBeConverted(SuperBPGC))
 				{
 					break;
 				}
@@ -1156,7 +1154,7 @@ FString FBlueprintCompilerCppBackendBase::GenerateWrapperForClass(UClass* Source
 		}
 
 		UBlueprintGeneratedClass* SuperBPGC = Cast<UBlueprintGeneratedClass>(SuperClassToUse);
-		if (SuperBPGC && !Dependencies.WillClassBeConverted(SuperBPGC))
+		if (SuperBPGC && !Dependencies->WillClassBeConverted(SuperBPGC))
 		{
 			ParentStruct = FString::Printf(TEXT("FUnconvertedWrapper__%s"), *FEmitHelper::GetCppName(SuperBPGC));
 			EmitterContext.MarkUnconvertedClassAsNecessary(SuperBPGC);
diff --git a/Engine/Source/Developer/BlueprintCompilerCppBackend/Private/BlueprintCompilerCppBackendGatherDependencies.cpp b/Engine/Source/Developer/BlueprintCompilerCppBackend/Private/BlueprintCompilerCppBackendGatherDependencies.cpp
--- a/Engine/Source/Developer/BlueprintCompilerCppBackend/Private/BlueprintCompilerCppBackendGatherDependencies.cpp
+++ b/Engine/Source/Developer/BlueprintCompilerCppBackend/Private/BlueprintCompilerCppBackendGatherDependencies.cpp
@@ -72,6 +72,14 @@ struct FGatherConvertedClassDependenciesHelperBase : public FReferenceCollector
 		}
 	}
 
+	void AddAssetDependency(UObject* InAsset)
+	{
+		if (InAsset)
+		{
+			Dependencies.Assets.Add(InAsset);
+		}
+	}
+
 	void AddConvertedClassDependency(UBlueprintGeneratedClass* InBPGC)
 	{
 		if (InBPGC && !Dependencies.ConvertedClasses.Contains(InBPGC))
@@ -166,13 +174,13 @@ struct FFindAssetsToInclude : public FGatherConvertedClassDependenciesHelperBase
 					else
 					{
 						// Add the class itself as a dependency.
-						Dependencies.Assets.AddUnique(OwnerClass);
+						AddAssetDependency(OwnerClass);
 
 						if (ObjAsBPGC)
 						{
 							// For BPGC types, we also include the CDO as a dependency (since it will be serialized).
 							// Note that if we get here, we already know from above that the BPGC is not being converted.
-							Dependencies.Assets.AddUnique(ObjAsBPGC->GetDefaultObject());
+							AddAssetDependency(ObjAsBPGC->GetDefaultObject());
 						}
 					}
 				}
@@ -186,19 +194,19 @@ struct FFindAssetsToInclude : public FGatherConvertedClassDependenciesHelperBase
 					else
 					{
 						// Add the struct itself as a dependency.
-						Dependencies.Assets.AddUnique(OwnerStruct);
+						AddAssetDependency(OwnerStruct);
 					}
 				}
 				else
 				{
 					// UFUNCTION, UENUM, etc.
-					Dependencies.Assets.AddUnique(Object);
+					AddAssetDependency(Object);
 				}
 			}
 			else
 			{
 				// Include the asset as a dependency.
-				Dependencies.Assets.AddUnique(Object);
+				AddAssetDependency(Object);
 			}
 
 			// No need to traverse these objects any further, so we just return.
@@ -317,7 +325,7 @@ struct FFindHeadersToInclude : public FGatherConvertedClassDependenciesHelperBas
 		}
 
 		{
-			auto ObjAsField = Cast<UField>(Object);
+			UField* ObjAsField = Cast<UField>(Object);
 			if (!ObjAsField)
 			{
 				const bool bTransientObject = (Object->HasAnyFlags(RF_Transient) && !Object->IsIn(CurrentlyConvertedStruct)) || Object->IsIn(GetTransientPackage());
@@ -358,7 +366,7 @@ struct FFindHeadersToInclude : public FGatherConvertedClassDependenciesHelperBas
 			return;
 		}
 
-		auto OwnedByAnythingInHierarchy = [&]()->bool
+		auto OwnedByAnythingInHierarchy = [this, Object, CurrentlyConvertedStruct]()->bool
 		{
 			for (UStruct* IterStruct = CurrentlyConvertedStruct; IterStruct; IterStruct = IterStruct->GetSuperStruct())
 			{
@@ -424,7 +432,7 @@ FGatherConvertedClassDependencies::FGatherConvertedClassDependencies(UStruct* In
 	static const FBoolConfigValueHelper DontNativizeDataOnlyBP(TEXT("BlueprintNativizationSettings"), TEXT("bDontNativizeDataOnlyBP"));
 	if (DontNativizeDataOnlyBP)
 	{
-		auto RemoveFieldsFromDataOnlyBP = [&](TSet<UField*>& FieldSet)
+		auto RemoveFieldsFromDataOnlyBP = [this](TSet<UField*>& FieldSet)
 		{
 			TSet<UField*> FieldsToAdd;
 			for (auto Iter = FieldSet.CreateIterator(); Iter; ++Iter)
@@ -447,7 +455,7 @@ FGatherConvertedClassDependencies::FGatherConvertedClassDependencies(UStruct* In
 
 	{
 		TSet<FName> ExcludedModules(NativizationOptions.ExcludedModules);
-		auto RemoveFieldsDependentOnExcludedModules = [&](TSet<UField*>& FieldSet)
+		auto RemoveFieldsDependentOnExcludedModules = [&ExcludedModules, InStruct](TSet<UField*>& FieldSet)
 		{
 			for (auto Iter = FieldSet.CreateIterator(); Iter; ++Iter)
 			{
@@ -463,9 +471,9 @@ FGatherConvertedClassDependencies::FGatherConvertedClassDependencies(UStruct* In
 		RemoveFieldsDependentOnExcludedModules(IncludeInBody);
 	}
 
-	auto GatherRequiredModules = [&](const TSet<UField*>& Fields)
+	auto GatherRequiredModules = [this](const TSet<UField*>& Fields)
 	{
-		for (auto Field : Fields)
+		for (UField* Field : Fields)
 		{
 			const UPackage* Package = Field ? Field->GetOutermost() : nullptr;
 			if (Package && Package->HasAnyPackageFlags(PKG_CompiledIn))
@@ -532,13 +540,13 @@ void FGatherConvertedClassDependencies::DependenciesForHeader()
 	TArray<UObject*> NeededObjects;
 	FReferenceFinder HeaderReferenceFinder(NeededObjects, nullptr, false, false, true, false);
 
-	auto ShouldIncludeHeaderFor = [&](UObject* InObj)->bool
+	auto ShouldIncludeHeaderFor = [this](UObject* InObj)->bool
 	{
 		if (InObj
 			&& (InObj->IsA<UClass>() || InObj->IsA<UEnum>() || InObj->IsA<UScriptStruct>())
 			&& !InObj->HasAnyFlags(RF_ClassDefaultObject))
 		{
-			auto ObjAsBPGC = Cast<UBlueprintGeneratedClass>(InObj);
+			UBlueprintGeneratedClass* ObjAsBPGC = Cast<UBlueprintGeneratedClass>(InObj);
 			const bool bWillBeConvetedAsBPGC = ObjAsBPGC && WillClassBeConverted(ObjAsBPGC);
 			const bool bRemainAsUnconvertedBPGC = ObjAsBPGC && !bWillBeConvetedAsBPGC;
 			if (!bRemainAsUnconvertedBPGC && (InObj->GetOutermost() != OriginalStruct->GetOutermost()))
@@ -549,7 +557,7 @@ void FGatherConvertedClassDependencies::DependenciesForHeader()
 		return false;
 	};
 
-	for (auto Obj : ObjectsToCheck)
+	for (UObject* Obj : ObjectsToCheck)
 	{
 		const UProperty* Property = Cast<const UProperty>(Obj);
 		if (const UArrayProperty* ArrayProperty = Cast<UArrayProperty>(Property))
@@ -561,28 +569,28 @@ void FGatherConvertedClassDependencies::DependenciesForHeader()
 		const bool bIsMemberVariable = OwnerProperty && (OwnerProperty->GetOuter() == OriginalStruct);
 		if (bIsParam || bIsMemberVariable)
 		{
-			if (auto SoftClassProperty = Cast<const USoftClassProperty>(Property))
+			if (const USoftClassProperty* SoftClassProperty = Cast<const USoftClassProperty>(Property))
 			{
 				DeclareInHeader.Add(GetFirstNativeOrConvertedClass(SoftClassProperty->MetaClass));
 			}
-			if (auto ClassProperty = Cast<const UClassProperty>(Property))
+			if (const UClassProperty* ClassProperty = Cast<const UClassProperty>(Property))
 			{
 				DeclareInHeader.Add(GetFirstNativeOrConvertedClass(ClassProperty->MetaClass));
 			}
-			if (auto ObjectProperty = Cast<const UObjectPropertyBase>(Property))
+			if (const UObjectPropertyBase* ObjectProperty = Cast<const UObjectPropertyBase>(Property))
 			{
 				DeclareInHeader.Add(GetFirstNativeOrConvertedClass(ObjectProperty->PropertyClass));
 			}
-			else if (auto InterfaceProperty = Cast<const UInterfaceProperty>(Property))
+			else if (const UInterfaceProperty* InterfaceProperty = Cast<const UInterfaceProperty>(Property))
 			{
 				IncludeInHeader.Add(InterfaceProperty->InterfaceClass);
 			}
-			else if (auto DelegateProperty = Cast<const UDelegateProperty>(Property))
+			else if (const UDelegateProperty* DelegateProperty = Cast<const UDelegateProperty>(Property))
 			{
 				IncludeInHeader.Add(DelegateProperty->SignatureFunction ? DelegateProperty->SignatureFunction->GetOwnerStruct() : nullptr);
 			}
 			/* MC Delegate signatures are recreated in local scope anyway.
-			else if (auto MulticastDelegateProperty = Cast<const UMulticastDelegateProperty>(Property))
+			else if (const UMulticastDelegateProperty* MulticastDelegateProperty = Cast<const UMulticastDelegateProperty>(Property))
 			{
 				IncludeInHeader.Add(MulticastDelegateProperty->SignatureFunction ? MulticastDelegateProperty->SignatureFunction->GetOwnerClass() : nullptr);
 			}
@@ -608,12 +616,12 @@ void FGatherConvertedClassDependencies::DependenciesForHeader()
 		}
 	}
 
-	if (auto SuperStruct = OriginalStruct->GetSuperStruct())
+	if (UStruct* SuperStruct = OriginalStruct->GetSuperStruct())
 	{
 		IncludeInHeader.Add(SuperStruct);
 	}
 
-	if (auto SourceClass = Cast<UClass>(OriginalStruct))
+	if (UClass* SourceClass = Cast<UClass>(OriginalStruct))
 	{
 		for (auto& ImplementedInterface : SourceClass->Interfaces)
 		{
@@ -621,7 +629,7 @@ void FGatherConvertedClassDependencies::DependenciesForHeader()
 		}
 	}
 
-	for (auto Obj : NeededObjects)
+	for (UObject* Obj : NeededObjects)
 	{
 		if (ShouldIncludeHeaderFor(Obj))
 		{
@@ -684,9 +692,9 @@ TSet<const UObject*> FGatherConvertedClassDependencies::AllDependencies() const
 		All.Add(SuperClass);
 	}
 
-	if (auto SourceClass = Cast<UClass>(OriginalStruct))
+	if (UClass* SourceClass = Cast<UClass>(OriginalStruct))
 	{
-		for (auto& ImplementedInterface : SourceClass->Interfaces)
+		for (const FImplementedInterface& ImplementedInterface : SourceClass->Interfaces)
 		{
 			UBlueprintGeneratedClass* InterfaceClass = Cast<UBlueprintGeneratedClass>(ImplementedInterface.Class);
 			if (ImplementedInterface.Class && (!InterfaceClass || WillClassBeConverted(InterfaceClass)))
@@ -695,11 +703,18 @@ TSet<const UObject*> FGatherConvertedClassDependencies::AllDependencies() const
 			}
 		}
 	}
+	
+	TSet<UObject*> NestedAssets;
+	GatherAssetsReferencedByConvertedTypes(NestedAssets);
 
 	for (auto It : Assets)
 	{
 		All.Add(It);
 	}
+	for (auto It : NestedAssets)
+	{
+		All.Add(It);
+	}
 	for (auto It : ConvertedClasses)
 	{
 		All.Add(It);
@@ -715,6 +730,24 @@ TSet<const UObject*> FGatherConvertedClassDependencies::AllDependencies() const
 	return All;
 }
 
+TMap<UStruct*, TSharedPtr<FGatherConvertedClassDependencies>> FGatherConvertedClassDependencies::CachedConvertedClassDependencies;
+
+TSharedPtr<FGatherConvertedClassDependencies> FGatherConvertedClassDependencies::Get(UStruct* InStruct, const FCompilerNativizationOptions& InNativizationOptions)
+{
+	TSharedPtr<FGatherConvertedClassDependencies> ConvertedClassDependenciesPtr;
+	if (InStruct)
+	{
+		ConvertedClassDependenciesPtr = CachedConvertedClassDependencies.FindOrAdd(InStruct);
+		if (!ConvertedClassDependenciesPtr.IsValid())
+		{
+			ConvertedClassDependenciesPtr = MakeShared<FGatherConvertedClassDependencies>(InStruct, InNativizationOptions);
+			check(ConvertedClassDependenciesPtr.IsValid());
+		}
+	}
+
+	return ConvertedClassDependenciesPtr;
+}
+
 class FArchiveReferencesInStructInstance : public FArchive
 {
 public:
@@ -745,6 +778,39 @@ public:
 	}
 };
 
+void FGatherConvertedClassDependencies::GatherAssetsReferencedByConvertedTypes(TSet<UObject*>& Dependencies) const
+{
+	TSet<UStruct*> VisitedTypes = { GetActualStruct() };
+
+	TArray<UStruct*> ConvertedTypeStack;
+	ConvertedTypeStack.Append(ConvertedStructs.Array());
+	ConvertedTypeStack.Append(ConvertedClasses.Array());
+
+	while(ConvertedTypeStack.Num() > 0)
+	{
+		UStruct* ConvertedType = ConvertedTypeStack.Pop();
+		TSharedPtr<FGatherConvertedClassDependencies> ConvertedTypeDependenciesPtr = Get(ConvertedType, NativizationOptions);
+
+		VisitedTypes.Add(ConvertedType);
+		Dependencies.Append(ConvertedTypeDependenciesPtr->Assets);
+
+		for (UUserDefinedStruct* ConvertedStruct : ConvertedTypeDependenciesPtr->ConvertedStructs)
+		{
+			if (!VisitedTypes.Contains(ConvertedStruct))
+			{
+				ConvertedTypeStack.Push(ConvertedStruct);
+			}
+		}
+		
+		for (UBlueprintGeneratedClass* ConvertedClass : ConvertedTypeDependenciesPtr->ConvertedClasses)
+		{
+			if (!VisitedTypes.Contains(ConvertedClass))
+			{
+				ConvertedTypeStack.Push(ConvertedClass);
+			}
+		}
+	}
+}
 
 void FGatherConvertedClassDependencies::GatherAssetsReferencedByUDSDefaultValue(TSet<UObject*>& Dependencies, UUserDefinedStruct* Struct)
 {
diff --git a/Engine/Source/Developer/BlueprintCompilerCppBackend/Private/BlueprintCompilerCppBackendUtils.cpp b/Engine/Source/Developer/BlueprintCompilerCppBackend/Private/BlueprintCompilerCppBackendUtils.cpp
--- a/Engine/Source/Developer/BlueprintCompilerCppBackend/Private/BlueprintCompilerCppBackendUtils.cpp
+++ b/Engine/Source/Developer/BlueprintCompilerCppBackend/Private/BlueprintCompilerCppBackendUtils.cpp
@@ -200,7 +200,7 @@ FString FEmitterLocalContext::FindGloballyMappedObject(const UObject* Object, co
 		if (bTryUsedAssetsList)
 		{
 			int32 AssetIndex = UsedObjectInCurrentClass.IndexOfByKey(Object);
-			if (INDEX_NONE == AssetIndex && Dependencies.Assets.Contains(Object))
+			if (INDEX_NONE == AssetIndex && Dependencies.Assets.Contains(const_cast<UObject*>(Object)))
 			{
 				AssetIndex = UsedObjectInCurrentClass.Add(Object);
 			}
diff --git a/Engine/Source/Developer/BlueprintCompilerCppBackend/Private/BlueprintCompilerCppBackendValueHelper.cpp b/Engine/Source/Developer/BlueprintCompilerCppBackend/Private/BlueprintCompilerCppBackendValueHelper.cpp
--- a/Engine/Source/Developer/BlueprintCompilerCppBackend/Private/BlueprintCompilerCppBackendValueHelper.cpp
+++ b/Engine/Source/Developer/BlueprintCompilerCppBackend/Private/BlueprintCompilerCppBackendValueHelper.cpp
@@ -1159,7 +1159,7 @@ struct FFakeImportTableHelper
 		UClass* SourceClass = Cast<UClass>(SourceStruct);
 		if (ensure(SourceStruct) && ensure(!SourceClass || OriginalClass))
 		{
-			auto GatherDependencies = [&](UStruct* InStruct)
+			auto GatherDependencies = [this](UStruct* InStruct)
 			{
 				SerializeBeforeSerializeStructDependencies.Add(InStruct->GetSuperStruct());
 
@@ -1217,7 +1217,7 @@ struct FFakeImportTableHelper
 				GatherDependencies(OriginalClass);
 			}
 
-			auto GetClassesOfSubobjects = [&](TMap<UObject*, FString>& SubobjectsMap)
+			auto GetClassesOfSubobjects = [this, &Context](TMap<UObject*, FString>& SubobjectsMap)
 			{
 				TArray<UObject*> Subobjects;
 				SubobjectsMap.GetKeys(Subobjects);
@@ -1225,8 +1225,21 @@ struct FFakeImportTableHelper
 				{
 					if (Subobject)
 					{
-						SerializeBeforeSerializeStructDependencies.Add(Subobject->GetClass());
-						SerializeBeforeCreateCDODependencies.Add(Subobject->GetClass()->GetDefaultObject());
+						UClass* SubobjectClass = Subobject->GetClass();
+						SerializeBeforeSerializeStructDependencies.Add(SubobjectClass);
+						SerializeBeforeCreateCDODependencies.Add(SubobjectClass->GetDefaultObject());
+
+						// This ensures that any nested asset dependencies will be serialized before attempting to instance a subobject that's a converted type when constructing the CDO.
+						if (UBlueprintGeneratedClass* SubobjectClassAsBPGC = Cast<UBlueprintGeneratedClass>(SubobjectClass))
+						{
+							if (Context.Dependencies.ConvertedClasses.Contains(SubobjectClassAsBPGC))
+							{
+								TSharedPtr<FGatherConvertedClassDependencies> SubobjectClassDependencies = FGatherConvertedClassDependencies::Get(SubobjectClassAsBPGC, Context.Dependencies.NativizationOptions);
+								
+								SerializeBeforeCreateCDODependencies.Append(SubobjectClassDependencies->Assets);
+								SubobjectClassDependencies->GatherAssetsReferencedByConvertedTypes(SerializeBeforeCreateCDODependencies);
+							}
+						}
 					}
 				}
 			};
@@ -1258,7 +1271,7 @@ struct FFakeImportTableHelper
 			//everything was created for class
 			CompactDataRef.CDODependency.bCreateBeforeCreateDependency = false; 
 
-			// Classes of subobjects, created while CDO construction
+			// Classes of subobjects, created while CDO construction, including assets they depend on for their own construction
 			CompactDataRef.CDODependency.bSerializationBeforeCreateDependency = SerializeBeforeCreateCDODependencies.Contains(const_cast<UObject*>(Asset));
 
 			// CDO is not serialized
diff --git a/Engine/Source/Developer/BlueprintCompilerCppBackend/Public/BlueprintCompilerCppBackendGatherDependencies.h b/Engine/Source/Developer/BlueprintCompilerCppBackend/Public/BlueprintCompilerCppBackendGatherDependencies.h
--- a/Engine/Source/Developer/BlueprintCompilerCppBackend/Public/BlueprintCompilerCppBackendGatherDependencies.h
+++ b/Engine/Source/Developer/BlueprintCompilerCppBackend/Public/BlueprintCompilerCppBackendGatherDependencies.h
@@ -13,12 +13,15 @@ class UUserDefinedStruct;
 /** The struct gathers dependencies of a converted BPGC */
 struct BLUEPRINTCOMPILERCPPBACKEND_API FGatherConvertedClassDependencies
 {
+private:
+	static TMap<UStruct*, TSharedPtr<FGatherConvertedClassDependencies>> CachedConvertedClassDependencies;
+
 protected:
 	UStruct* OriginalStruct;
 
 public:
 	// Dependencies:
-	TArray<UObject*> Assets; 
+	TSet<UObject*> Assets;
 
 	TSet<UBlueprintGeneratedClass*> ConvertedClasses;
 	TSet<UUserDefinedStruct*> ConvertedStructs;
@@ -34,7 +37,7 @@ public:
 	FCompilerNativizationOptions NativizationOptions;
 
 public:
-	FGatherConvertedClassDependencies(UStruct* InStruct, const FCompilerNativizationOptions& InNativizationOptions);
+	static TSharedPtr<FGatherConvertedClassDependencies> Get(UStruct* InStruct, const FCompilerNativizationOptions& InNativizationOptions);
 
 	UStruct* GetActualStruct() const
 	{
@@ -50,10 +53,17 @@ public:
 public:
 	bool WillClassBeConverted(const UBlueprintGeneratedClass* InClass) const;
 
+	void GatherAssetsReferencedByConvertedTypes(TSet<UObject*>& Dependencies) const;
+
 	static void GatherAssetsReferencedByUDSDefaultValue(TSet<UObject*>& Dependencies, UUserDefinedStruct* Struct);
 
 	static bool IsFieldFromExcludedPackage(const UField* Field, const TSet<FName>& InExcludedModules);
 
 protected:
+	FGatherConvertedClassDependencies(UStruct* InStruct, const FCompilerNativizationOptions& InNativizationOptions);
+
+	/** Friend for access to constructor via MakeShared */
+	friend class SharedPointerInternals::TIntrusiveReferenceController<FGatherConvertedClassDependencies>;
+
 	void DependenciesForHeader();
 };
diff --git a/Engine/Source/Developer/BlueprintNativeCodeGen/Private/NativeCodeGenerationTool.cpp b/Engine/Source/Developer/BlueprintNativeCodeGen/Private/NativeCodeGenerationTool.cpp
--- a/Engine/Source/Developer/BlueprintNativeCodeGen/Private/NativeCodeGenerationTool.cpp
+++ b/Engine/Source/Developer/BlueprintNativeCodeGen/Private/NativeCodeGenerationTool.cpp
@@ -66,16 +66,16 @@ struct FGeneratedCodeData
 	void GatherUserDefinedDependencies(UBlueprint& InBlueprint)
 	{
 		FCompilerNativizationOptions BlankOptions{};
-		FGatherConvertedClassDependencies ClassDependencies(InBlueprint.GeneratedClass, BlankOptions);
-		for (auto Iter : ClassDependencies.ConvertedClasses)
+		TSharedPtr<FGatherConvertedClassDependencies> ClassDependencies = FGatherConvertedClassDependencies::Get(InBlueprint.GeneratedClass, BlankOptions);
+		for (auto Iter : ClassDependencies->ConvertedClasses)
 		{
 			DependentObjects.Add(Iter);
 		}
-		for (auto Iter : ClassDependencies.ConvertedStructs)
+		for (auto Iter : ClassDependencies->ConvertedStructs)
 		{
 			DependentObjects.Add(Iter);
 		}
-		for (auto Iter : ClassDependencies.ConvertedEnum)
+		for (auto Iter : ClassDependencies->ConvertedEnum)
 		{
 			DependentObjects.Add(Iter);
 		}
@@ -96,7 +96,7 @@ struct FGeneratedCodeData
 		DependentObjects.Add(InBlueprint.GeneratedClass);
 
 		bool bUnconvertedHeader = false;
-		for (auto Asset : ClassDependencies.Assets)
+		for (auto Asset : ClassDependencies->Assets)
 		{
 			if (auto BPGC = Cast<UBlueprintGeneratedClass>(Asset))
 			{
