Misc. cleanup of FEmitDefaultValueHelper::AddStaticFunctionsForDependencies().

Change summary:
- Replaced inclusive lambda capture lists and non-compliant occurrences of auto.
- Consolidated EDL boot time paths and replaced TODO comments with updated notes.
- Consolidated work that required a full iteration through the dependencies list into a single pass.

#rnx

diff --git a/Engine/Source/Developer/BlueprintCompilerCppBackend/Private/BlueprintCompilerCppBackendValueHelper.cpp b/Engine/Source/Developer/BlueprintCompilerCppBackend/Private/BlueprintCompilerCppBackendValueHelper.cpp
--- a/Engine/Source/Developer/BlueprintCompilerCppBackend/Private/BlueprintCompilerCppBackendValueHelper.cpp
+++ b/Engine/Source/Developer/BlueprintCompilerCppBackend/Private/BlueprintCompilerCppBackendValueHelper.cpp
@@ -1285,33 +1285,8 @@ void FEmitDefaultValueHelper::AddStaticFunctionsForDependencies(FEmitterLocalCon
 	, TSharedPtr<FGatherConvertedClassDependencies> ParentDependencies
 	, FCompilerNativizationOptions NativizationOptions)
 {
-	// 1. GATHER UDS DEFAULT VALUE DEPENDENCIES
-	{
-		TSet<UObject*> References;
-		for (UUserDefinedStruct* UDS : Context.StructsWithDefaultValuesUsed)
-		{
-			FGatherConvertedClassDependencies::GatherAssetsReferencedByUDSDefaultValue(References, UDS);
-		}
-		for (UObject* Obj : References)
-		{
-			Context.UsedObjectInCurrentClass.AddUnique(Obj);
-		}
-	}
-
-	// 2. ALL ASSETS TO LIST
-	TSet<const UObject*> AllDependenciesToHandle = Context.Dependencies.AllDependencies();
-	AllDependenciesToHandle.Append(Context.UsedObjectInCurrentClass);
-	AllDependenciesToHandle.Remove(nullptr);
-
-	// Special case, we don't need to load any dependencies from CoreUObject.
-	UPackage* CoreUObjectPackage = UProperty::StaticClass()->GetOutermost();
-	for (auto Iter = AllDependenciesToHandle.CreateIterator(); Iter; ++Iter)
-	{
-		if ((*Iter)->GetOutermost() == CoreUObjectPackage)
-		{
-			Iter.RemoveCurrent();
-		}
-	}
+	constexpr bool bBootTimeEDL = USE_EVENT_DRIVEN_ASYNC_LOAD_AT_BOOT_TIME;
+	const bool bEnableBootTimeEDLOptimization = IsEventDrivenLoaderEnabledInCookedBuilds() && bBootTimeEDL;
 
 	// HELPERS
 	UStruct* SourceStruct = Context.Dependencies.GetActualStruct();
@@ -1323,7 +1298,7 @@ void FEmitDefaultValueHelper::AddStaticFunctionsForDependencies(FEmitterLocalCon
 	const FString CppTypeName = FEmitHelper::GetCppName(SourceStruct);
 	FFakeImportTableHelper FakeImportTableHelper(SourceStruct, OriginalClass, Context);
 
-	auto CreateAssetToLoadString = [&](const UObject* AssetObj) -> FString
+	auto CreateAssetToLoadString = [&Context](const UObject* AssetObj) -> FString
 	{
 		UClass* AssetType = AssetObj->GetClass();
 		if (AssetType->IsChildOf<UUserDefinedEnum>())
@@ -1356,7 +1331,7 @@ void FEmitDefaultValueHelper::AddStaticFunctionsForDependencies(FEmitterLocalCon
 			, *OuterName);
 	};
 
-	auto CreateDependencyRecord = [&](const UObject* InAsset, FString& OptionalComment) -> FCompactBlueprintDependencyData
+	auto CreateDependencyRecord = [&NativizationOptions, &FakeImportTableHelper, &CppTypeName, OriginalClass, &CreateAssetToLoadString, bEnableBootTimeEDLOptimization](const UObject* InAsset, FString& OptionalComment) -> FCompactBlueprintDependencyData
 	{
 		ensure(InAsset);
 		if (InAsset && IsEditorOnlyObject(InAsset))
@@ -1400,9 +1375,8 @@ void FEmitDefaultValueHelper::AddStaticFunctionsForDependencies(FEmitterLocalCon
 		FakeImportTableHelper.FillDependencyData(InAsset, Result);
 		return Result;
 	};
-	const bool bBootTimeEDL = USE_EVENT_DRIVEN_ASYNC_LOAD_AT_BOOT_TIME;
-	const bool bEnableBootTimeEDLOptimization = IsEventDrivenLoaderEnabledInCookedBuilds() && bBootTimeEDL;
-	auto AddAssetArray = [&](const TArray<const UObject*>& Assets)
+	
+	auto AddAssetArray = [&Context, SourceStruct, &CreateDependencyRecord, bEnableBootTimeEDLOptimization](const TArray<const UObject*>& Assets)
 	{
 		if (Assets.Num())
 		{
@@ -1456,23 +1430,77 @@ void FEmitDefaultValueHelper::AddStaticFunctionsForDependencies(FEmitterLocalCon
 		}
 	};
 
+	// 1. GATHER UDS DEFAULT VALUE DEPENDENCIES
+	{
+		TSet<UObject*> References;
+		for (UUserDefinedStruct* UDS : Context.StructsWithDefaultValuesUsed)
+		{
+			FGatherConvertedClassDependencies::GatherAssetsReferencedByUDSDefaultValue(References, UDS);
+		}
+		for (UObject* Obj : References)
+		{
+			Context.UsedObjectInCurrentClass.AddUnique(Obj);
+		}
+	}
+
+	// 2. ALL ASSETS TO LIST
 	TSet<const UBlueprintGeneratedClass*> OtherBPGCs;
-	if (!bEnableBootTimeEDLOptimization)
+	TSet<const UObject*> AllDependenciesToHandle = Context.Dependencies.AllDependencies();
 	{
-		for (const UObject* It : AllDependenciesToHandle)
+		// Append used objects.
+		AllDependenciesToHandle.Append(Context.UsedObjectInCurrentClass);
+
+		// Remove invalid dependencies.
+		AllDependenciesToHandle.Remove(nullptr);
+		
+		// Remove unnecessary dependencies.
+		for (auto Iter = AllDependenciesToHandle.CreateIterator(); Iter; ++Iter)
 		{
-			if (const UBlueprintGeneratedClass* OtherBPGC = Cast<const UBlueprintGeneratedClass>(It))
+			bool bCanExclude = false;
+
+			if (const UObject* ItObj = *Iter)
 			{
-				const UBlueprint* BP = Cast<const UBlueprint>(OtherBPGC->ClassGeneratedBy);
-				if (Context.Dependencies.WillClassBeConverted(OtherBPGC) && BP && (BP->BlueprintType != EBlueprintType::BPTYPE_Interface))
+				// Special case, we don't need to load any dependencies from CoreUObject.
+				static const UPackage* CoreUObjectPackage = UProperty::StaticClass()->GetOutermost();
+				bCanExclude = ItObj->GetOutermost() == CoreUObjectPackage;
+
+				// We can exclude native type dependencies if EDL is not going to be enabled at boot time.
+				if (!bCanExclude && !bEnableBootTimeEDLOptimization)
 				{
-					OtherBPGCs.Add(OtherBPGC);
+					if (const UClass* ObjAsClass = Cast<const UClass>(ItObj))
+					{
+						if (ObjAsClass->HasAnyClassFlags(CLASS_Native))
+						{
+							bCanExclude = true;
+						}
+						else if (const UBlueprintGeneratedClass* OtherBPGC = Cast<const UBlueprintGeneratedClass>(ObjAsClass))
+						{
+							// Gather the set of all non-native, non-interface class dependencies that will be converted. This is used below to help reduce code size when the EDL will not be enabled at boot time.
+							const UBlueprint* BP = Cast<const UBlueprint>(OtherBPGC->ClassGeneratedBy);
+							if (Context.Dependencies.WillClassBeConverted(OtherBPGC) && BP && (BP->BlueprintType != EBlueprintType::BPTYPE_Interface))
+							{
+								OtherBPGCs.Add(OtherBPGC);
+							}
+						}
+					}
+					else
+					{
+						// Exclude native UENUM() types that are not user-defined.
+						bCanExclude |= (ItObj->IsA<UEnum>() && !ItObj->IsA<UUserDefinedEnum>());
+
+						// Exclude native USTRUCT() types that are not user-defined.
+						bCanExclude |= (ItObj->IsA<UScriptStruct>() && !ItObj->IsA<UUserDefinedStruct>());
+					}
 				}
 			}
+
+			if (bCanExclude)
+			{
+				Iter.RemoveCurrent();
+			}
 		}
 	}
 
-
 	// 3. LIST OF UsedAssets
 	if (SourceStruct->IsA<UClass>())
 	{
@@ -1485,9 +1513,11 @@ void FEmitDefaultValueHelper::AddStaticFunctionsForDependencies(FEmitterLocalCon
 		for (int32 UsedAssetIndex = 0; UsedAssetIndex < Context.UsedObjectInCurrentClass.Num(); ++UsedAssetIndex)
 		{
 			const UObject* LocAsset = Context.UsedObjectInCurrentClass[UsedAssetIndex];
-			ensure(AllDependenciesToHandle.Contains(LocAsset));
-			AssetsToAdd.Add(LocAsset);
-			AllDependenciesToHandle.Remove(LocAsset);
+			if (AllDependenciesToHandle.Contains(LocAsset))
+			{
+				AssetsToAdd.Add(LocAsset);
+				AllDependenciesToHandle.Remove(LocAsset);
+			}
 		}
 		AddAssetArray(AssetsToAdd);
 		Context.DecreaseIndent();
@@ -1541,34 +1571,6 @@ void FEmitDefaultValueHelper::AddStaticFunctionsForDependencies(FEmitterLocalCon
 				Context.AddLine(TEXT("}"));
 			}
 		}
-
-		if (bEnableBootTimeEDLOptimization)
-		{
-			//TODO: remove stuff from CoreUObject
-		}
-		else
-		{
-			//WIthout EDL we don't need the native stuff.
-			for (auto Iter = AllDependenciesToHandle.CreateIterator(); Iter; ++Iter)
-			{
-				const UObject* ItObj = *Iter;
-				if (auto ObjAsClass = Cast<const UClass>(ItObj))
-				{
-					if (ObjAsClass->HasAnyClassFlags(CLASS_Native))
-					{
-						Iter.RemoveCurrent();
-					}
-				}
-				else if (ItObj && ItObj->IsA<UScriptStruct>() && !ItObj->IsA<UUserDefinedStruct>())
-				{
-					Iter.RemoveCurrent();
-				}
-				else if (ItObj && ItObj->IsA<UEnum>() && !ItObj->IsA<UUserDefinedEnum>())
-				{
-					Iter.RemoveCurrent();
-				}
-			}
-		}
 		
 		AddAssetArray(AllDependenciesToHandle.Array());
 		Context.DecreaseIndent();
