Fix invalid cast of a nativized enum term to its underlying pointer type in generated Blueprint C++ code.

diff --git a/Engine/Source/Developer/BlueprintCompilerCppBackend/Private/BlueprintCompilerCppBackendUtils.cpp b/Engine/Source/Developer/BlueprintCompilerCppBackend/Private/BlueprintCompilerCppBackendUtils.cpp
--- a/Engine/Source/Developer/BlueprintCompilerCppBackend/Private/BlueprintCompilerCppBackendUtils.cpp
+++ b/Engine/Source/Developer/BlueprintCompilerCppBackend/Private/BlueprintCompilerCppBackendUtils.cpp
@@ -1604,7 +1604,8 @@ bool FEmitHelper::GenerateAutomaticCast(FEmitterLocalContext& EmitterContext, co
 
 				if (bForceReference)
 				{
-					OutCastBegin = TEXT("*static_cast<uint8*>(&(");
+					// An enum and its underlying type are not related by inheritance, so 'static_cast' cannot be used here (5.2.9/11).
+					OutCastBegin = TEXT("*reinterpret_cast<uint8*>(&(");
 					OutCastEnd = TEXT("))");
 				}
 				else
