Fix spelling error
#rnx

diff --git a/Engine/Source/Developer/BlueprintCompilerCppBackend/Private/BlueprintCompilerCppBackendBase.cpp b/Engine/Source/Developer/BlueprintCompilerCppBackend/Private/BlueprintCompilerCppBackendBase.cpp
--- a/Engine/Source/Developer/BlueprintCompilerCppBackend/Private/BlueprintCompilerCppBackendBase.cpp
+++ b/Engine/Source/Developer/BlueprintCompilerCppBackend/Private/BlueprintCompilerCppBackendBase.cpp
@@ -74,7 +74,7 @@ void FBlueprintCompilerCppBackendBase::EmitStructProperties(FEmitterLocalContext
 		check(Property);
 		FString PropertyMacro(TEXT("UPROPERTY("));
 		{
-			TArray<FString> Tags = FEmitHelper::ProperyFlagsToTags(Property->PropertyFlags, nullptr != Cast<UClass>(SourceClass));
+			TArray<FString> Tags = FEmitHelper::PropertyFlagsToTags(Property->PropertyFlags, nullptr != Cast<UClass>(SourceClass));
 			Tags.Emplace(FEmitHelper::HandleRepNotifyFunc(Property));
 			Tags.Emplace(FEmitHelper::HandleMetaData(Property, false));
 			Tags.Remove(FString());
diff --git a/Engine/Source/Developer/BlueprintCompilerCppBackend/Private/BlueprintCompilerCppBackendUtils.cpp b/Engine/Source/Developer/BlueprintCompilerCppBackend/Private/BlueprintCompilerCppBackendUtils.cpp
--- a/Engine/Source/Developer/BlueprintCompilerCppBackend/Private/BlueprintCompilerCppBackendUtils.cpp
+++ b/Engine/Source/Developer/BlueprintCompilerCppBackend/Private/BlueprintCompilerCppBackendUtils.cpp
@@ -658,7 +658,7 @@ FString FEmitHelper::HandleMetaData(const UField* Field, bool AddCategory, const
 #endif
 #define HANDLE_CPF_TAG(TagName, CheckedFlags) if (HasAllFlags(Flags, (CheckedFlags))) { Tags.Emplace(TagName); }
 
-TArray<FString> FEmitHelper::ProperyFlagsToTags(uint64 Flags, bool bIsClassProperty)
+TArray<FString> FEmitHelper::PropertyFlagsToTags(uint64 Flags, bool bIsClassProperty)
 {
 	TArray<FString> Tags;
 
diff --git a/Engine/Source/Developer/BlueprintCompilerCppBackend/Private/BlueprintCompilerCppBackendUtils.h b/Engine/Source/Developer/BlueprintCompilerCppBackend/Private/BlueprintCompilerCppBackendUtils.h
--- a/Engine/Source/Developer/BlueprintCompilerCppBackend/Private/BlueprintCompilerCppBackendUtils.h
+++ b/Engine/Source/Developer/BlueprintCompilerCppBackend/Private/BlueprintCompilerCppBackendUtils.h
@@ -267,7 +267,7 @@ struct FEmitHelper
 
 	static FString HandleMetaData(const UField* Field, bool AddCategory = true, const TArray<FString>* AdditinalMetaData = nullptr);
 
-	static TArray<FString> ProperyFlagsToTags(uint64 Flags, bool bIsClassProperty);
+	static TArray<FString> PropertyFlagsToTags(uint64 Flags, bool bIsClassProperty);
 
 	static TArray<FString> FunctionFlagsToTags(uint64 Flags);
 
