auto removal
#rnx

diff --git a/Engine/Source/Developer/BlueprintCompilerCppBackend/Private/BlueprintCompilerCppBackendAnim.cpp b/Engine/Source/Developer/BlueprintCompilerCppBackend/Private/BlueprintCompilerCppBackendAnim.cpp
--- a/Engine/Source/Developer/BlueprintCompilerCppBackend/Private/BlueprintCompilerCppBackendAnim.cpp
+++ b/Engine/Source/Developer/BlueprintCompilerCppBackend/Private/BlueprintCompilerCppBackendAnim.cpp
@@ -19,16 +19,16 @@ void FBackendHelperAnim::AddHeaders(FEmitterLocalContext& Context)
 
 void FBackendHelperAnim::CreateAnimClassData(FEmitterLocalContext& Context)
 {
-	if (auto AnimClass = Cast<UAnimBlueprintGeneratedClass>(Context.GetCurrentlyGeneratedClass()))
+	if (UAnimBlueprintGeneratedClass* AnimClass = Cast<UAnimBlueprintGeneratedClass>(Context.GetCurrentlyGeneratedClass()))
 	{
 		const FString LocalNativeName = Context.GenerateUniqueLocalName();
-		Context.AddLine(FString::Printf(TEXT("auto %s = NewObject<UAnimClassData>(InDynamicClass, TEXT(\"AnimClassData\"));"), *LocalNativeName));
+		Context.AddLine(FString::Printf(TEXT("UAnimClassData* %s = NewObject<UAnimClassData>(InDynamicClass, TEXT(\"AnimClassData\"));"), *LocalNativeName));
 
-		auto AnimClassData = NewObject<UAnimClassData>(GetTransientPackage(), TEXT("AnimClassData"));
+		UAnimClassData* AnimClassData = NewObject<UAnimClassData>(GetTransientPackage(), TEXT("AnimClassData"));
 		AnimClassData->CopyFrom(AnimClass);
 
-		auto ObjectArchetype = AnimClassData->GetArchetype();
-		for (auto Property : TFieldRange<const UProperty>(UAnimClassData::StaticClass()))
+		UObject* ObjectArchetype = AnimClassData->GetArchetype();
+		for (const UProperty* Property : TFieldRange<const UProperty>(UAnimClassData::StaticClass()))
 		{
 			FEmitDefaultValueHelper::OuterGenerate(Context, Property, LocalNativeName
 				, reinterpret_cast<const uint8*>(AnimClassData)
