Fix spelling and remove autos
#rnx

diff --git a/Engine/Source/Developer/BlueprintCompilerCppBackend/Private/BlueprintCompilerCppBackendUtils.h b/Engine/Source/Developer/BlueprintCompilerCppBackend/Private/BlueprintCompilerCppBackendUtils.h
--- a/Engine/Source/Developer/BlueprintCompilerCppBackend/Private/BlueprintCompilerCppBackendUtils.h
+++ b/Engine/Source/Developer/BlueprintCompilerCppBackend/Private/BlueprintCompilerCppBackendUtils.h
@@ -13,7 +13,7 @@ class USCS_Node;
 class UUserDefinedEnum;
 class UUserDefinedStruct;
 struct FDefaultSubobjectData;
-struct FNonativeComponentData;
+struct FNonNativeComponentData;
 enum class ENativizedTermUsage : uint8;
 struct FEmitterLocalContext;
 
@@ -333,8 +333,6 @@ struct FEmitHelper
 	static const TCHAR* EmptyDefaultConstructor(UScriptStruct* Struct);
 };
 
-struct FNonativeComponentData;
-
 struct FEmitDefaultValueHelper
 {
 	static void GenerateUserStructConstructor(const UUserDefinedStruct* Struct, FEmitterLocalContext& EmitterContext);
@@ -390,7 +388,7 @@ private:
 	static FString HandleSpecialTypes(FEmitterLocalContext& Context, const UProperty* Property, const uint8* ValuePtr);
 
 	static FString HandleNonNativeComponent(FEmitterLocalContext& Context, const USCS_Node* Node, TSet<const UProperty*>& OutHandledProperties
-		, TArray<FString>& NativeCreatedComponentProperties, const USCS_Node* ParentNode, TArray<FNonativeComponentData>& ComponentsToInit
+		, TArray<FString>& NativeCreatedComponentProperties, const USCS_Node* ParentNode, TArray<FNonNativeComponentData>& ComponentsToInit
 		, bool bBlockRecursion);
 };
 
diff --git a/Engine/Source/Developer/BlueprintCompilerCppBackend/Private/BlueprintCompilerCppBackendValueHelper.cpp b/Engine/Source/Developer/BlueprintCompilerCppBackend/Private/BlueprintCompilerCppBackendValueHelper.cpp
--- a/Engine/Source/Developer/BlueprintCompilerCppBackend/Private/BlueprintCompilerCppBackendValueHelper.cpp
+++ b/Engine/Source/Developer/BlueprintCompilerCppBackend/Private/BlueprintCompilerCppBackendValueHelper.cpp
@@ -156,7 +156,7 @@ void FEmitDefaultValueHelper::GenerateUserStructConstructor(const UUserDefinedSt
 		TGuardValue<FCodeText*> OriginalDefaultTarget(Context.DefaultTarget, &Context.Body);
 		FStructOnScope StructData(Struct);
 		FUserStructOnScopeIgnoreDefaults RawDefaultStructOnScope(Struct);
-		for (auto Property : TFieldRange<const UProperty>(Struct))
+		for (const UProperty* Property : TFieldRange<const UProperty>(Struct))
 		{
 			// Since UDS types are converted to native USTRUCT, all POD fields must be initialized in the ctor, just as with "regular" native USTRUCT types.
 			const bool bForceInit = Property->HasAnyPropertyFlags(CPF_IsPlainOldData);
@@ -253,7 +253,7 @@ void FEmitDefaultValueHelper::InnerGenerate(FEmitterLocalContext& Context, const
 
 		const FString ActualPathToMember = InnerInlineStruct ? FString::Printf(TEXT("((%s*)%s.GetPtr())"), *FEmitHelper::GetCppName(InnerInlineStruct), *PathToMember) : PathToMember;
 
-		for (auto LocalProperty : TFieldRange<const UProperty>(ActualStruct))
+		for (const UProperty* LocalProperty : TFieldRange<const UProperty>(ActualStruct))
 		{
 			OuterGenerate(Context, LocalProperty, ActualPathToMember, ActualValuePtr
 				, (ActualDefaultValuePtr ? ActualDefaultValuePtr : DefaultStructOnScope.GetStructMemory())
@@ -518,9 +518,9 @@ bool FEmitDefaultValueHelper::SpecialStructureConstructor(const UStruct* Struct,
 		if (OutResult)
 		{
 			const FTransform* Transform = reinterpret_cast<const FTransform*>(ValuePtr);
-			const auto Rotation = Transform->GetRotation();
-			const auto Translation = Transform->GetTranslation();
-			const auto Scale = Transform->GetScale3D();
+			const FQuat Rotation = Transform->GetRotation();
+			const FVector Translation = Transform->GetTranslation();
+			const FVector Scale = Transform->GetScale3D();
 			*OutResult = FString::Printf(TEXT("FTransform( FQuat(%s,%s,%s,%s), FVector(%s,%s,%s), FVector(%s,%s,%s) )"),
 				*FEmitHelper::FloatToString(Rotation.X), *FEmitHelper::FloatToString(Rotation.Y), *FEmitHelper::FloatToString(Rotation.Z), *FEmitHelper::FloatToString(Rotation.W),
 				*FEmitHelper::FloatToString(Translation.X), *FEmitHelper::FloatToString(Translation.Y), *FEmitHelper::FloatToString(Translation.Z),
@@ -879,7 +879,7 @@ struct FDefaultSubobjectData
 		// Now walk through the property list and initialize delta values for this instance. Any nested instanced default
 		// subobjects found above that are also assigned to a reference property will be correctly seen as already handled.
 		const UClass* ObjectClass = Object->GetClass();
-		for (auto Property : TFieldRange<const UProperty>(ObjectClass))
+		for (const UProperty* Property : TFieldRange<const UProperty>(ObjectClass))
 		{
 			if (!HandledAsSpecialProperty(Context, Property))
 			{
@@ -986,7 +986,7 @@ protected:
 	}
 };
 
-struct FNonativeComponentData : public FDefaultSubobjectData
+struct FNonNativeComponentData : public FDefaultSubobjectData
 {
 	////
 	const USCS_Node* SCSNode;
@@ -994,13 +994,13 @@ struct FNonativeComponentData : public FDefaultSubobjectData
 	/** Socket/Bone that Component might attach to */
 	FName AttachToName;
 
-	FNonativeComponentData()
+	FNonNativeComponentData()
 		: SCSNode(nullptr)
 	{
 		bAddLocalScope = false;
 	}
 
-	virtual ~FNonativeComponentData()
+	virtual ~FNonNativeComponentData()
 	{
 	}
 
@@ -1029,7 +1029,7 @@ struct FNonativeComponentData : public FDefaultSubobjectData
 
 FString FEmitDefaultValueHelper::HandleNonNativeComponent(FEmitterLocalContext& Context, const USCS_Node* Node
 	, TSet<const UProperty*>& OutHandledProperties, TArray<FString>& NativeCreatedComponentProperties
-	, const USCS_Node* ParentNode, TArray<FNonativeComponentData>& ComponentsToInit
+	, const USCS_Node* ParentNode, TArray<FNonNativeComponentData>& ComponentsToInit
 	, bool bBlockRecursion)
 {
 	check(Node);
@@ -1062,10 +1062,10 @@ FString FEmitDefaultValueHelper::HandleNonNativeComponent(FEmitterLocalContext&
 
 		if (ComponentTemplate->GetOuter() == BPGC)
 		{
-			FNonativeComponentData NonativeComponentData;
-			NonativeComponentData.SCSNode = Node;
-			NonativeComponentData.VariableName = NativeVariablePropertyName;
-			NonativeComponentData.Object = ComponentTemplate;
+			FNonNativeComponentData NonNativeComponentData;
+			NonNativeComponentData.SCSNode = Node;
+			NonNativeComponentData.VariableName = NativeVariablePropertyName;
+			NonNativeComponentData.Object = ComponentTemplate;
 			UClass* ComponentClass = ComponentTemplate->GetClass();
 			check(ComponentClass != nullptr);
 
@@ -1083,7 +1083,7 @@ FString FEmitDefaultValueHelper::HandleNonNativeComponent(FEmitterLocalContext&
 					, *FEmitHelper::GetCppName(ComponentClass)
 					, *VariableCleanName));
 
-				NonativeComponentData.bWasCreated = true;
+				NonNativeComponentData.bWasCreated = true;
 				NativeCreatedComponentProperties.Add(NativeVariablePropertyName);
 
 				FString ParentVariableName;
@@ -1097,18 +1097,18 @@ FString FEmitDefaultValueHelper::HandleNonNativeComponent(FEmitterLocalContext&
 				{
 					ParentVariableName = Context.FindGloballyMappedObject(ParentComponentTemplate, USceneComponent::StaticClass());
 				}
-				NonativeComponentData.ParentVariableName = ParentVariableName;
-				NonativeComponentData.AttachToName = Node->AttachToName;
+				NonNativeComponentData.ParentVariableName = ParentVariableName;
+				NonNativeComponentData.AttachToName = Node->AttachToName;
 			}
-			NonativeComponentData.Archetype = ObjectToCompare;
-			ComponentsToInit.Add(NonativeComponentData);
+			NonNativeComponentData.Archetype = ObjectToCompare;
+			ComponentsToInit.Add(NonNativeComponentData);
 		}
 	}
 
 	// Recursively handle child nodes.
 	if (!bBlockRecursion)
 	{
-		for (auto ChildNode : Node->ChildNodes)
+		for (USCS_Node* ChildNode : Node->ChildNodes)
 		{
 			HandleNonNativeComponent(Context, ChildNode, OutHandledProperties, NativeCreatedComponentProperties, Node, ComponentsToInit, bBlockRecursion);
 		}
@@ -1631,7 +1631,7 @@ void FEmitDefaultValueHelper::AddRegisterHelper(FEmitterLocalContext& Context)
 
 void FEmitDefaultValueHelper::GenerateCustomDynamicClassInitialization(FEmitterLocalContext& Context, TSharedPtr<FGatherConvertedClassDependencies> ParentDependencies)
 {
-	auto BPGC = CastChecked<UBlueprintGeneratedClass>(Context.GetCurrentlyGeneratedClass());
+	UBlueprintGeneratedClass* BPGC = CastChecked<UBlueprintGeneratedClass>(Context.GetCurrentlyGeneratedClass());
 	const FString CppClassName = FEmitHelper::GetCppName(BPGC);
 
 	{
@@ -1655,7 +1655,7 @@ void FEmitDefaultValueHelper::GenerateCustomDynamicClassInitialization(FEmitterL
 		{
 			Context.AddLine(TEXT("// List of all referenced converted enums"));
 		}
-		for (auto LocEnum : Context.Dependencies.ConvertedEnum)
+		for (UUserDefinedEnum* LocEnum : Context.Dependencies.ConvertedEnum)
 		{
 			Context.AddLine(FString::Printf(TEXT("InDynamicClass->%s.Add(LoadObject<UEnum>(nullptr, TEXT(\"%s\")));"), GET_MEMBER_NAME_STRING_CHECKED(UDynamicClass, ReferencedConvertedFields), *(LocEnum->GetPathName().ReplaceCharWithEscapedChar())));
 			Context.EnumsInCurrentClass.Add(LocEnum);
@@ -1665,7 +1665,7 @@ void FEmitDefaultValueHelper::GenerateCustomDynamicClassInitialization(FEmitterL
 		{
 			Context.AddLine(TEXT("// List of all referenced converted classes"));
 		}
-		for (auto LocStruct : Context.Dependencies.ConvertedClasses)
+		for (UBlueprintGeneratedClass* LocStruct : Context.Dependencies.ConvertedClasses)
 		{
 			UClass* ClassToLoad = Context.Dependencies.FindOriginalClass(LocStruct);
 			if (ensure(ClassToLoad))
@@ -1697,7 +1697,7 @@ void FEmitDefaultValueHelper::GenerateCustomDynamicClassInitialization(FEmitterL
 		{
 			Context.AddLine(TEXT("// List of all referenced converted structures"));
 		}
-		for (auto LocStruct : Context.Dependencies.ConvertedStructs)
+		for (UUserDefinedStruct* LocStruct : Context.Dependencies.ConvertedStructs)
 		{
 			if (ParentDependencies.IsValid() && ParentDependencies->ConvertedStructs.Contains(LocStruct))
 			{
@@ -1710,7 +1710,7 @@ void FEmitDefaultValueHelper::GenerateCustomDynamicClassInitialization(FEmitterL
 
 		TArray<UActorComponent*> ActorComponentTempatesOwnedByClass = BPGC->ComponentTemplates;
 		// Gather all CT from SCS and IH, the remaining ones are generated for class..
-		if (auto SCS = BPGC->SimpleConstructionScript)
+		if (USimpleConstructionScript* SCS = BPGC->SimpleConstructionScript)
 		{
 			// >>> This code should be removed, once UE-39168 is fixed
 			//TODO: it's an ugly workaround - template from DefaultSceneRootNode is unnecessarily cooked :(
@@ -1721,12 +1721,12 @@ void FEmitDefaultValueHelper::GenerateCustomDynamicClassInitialization(FEmitterL
 			}
 			// <<< This code should be removed, once UE-39168 is fixed
 
-			for (auto Node : SCS->GetAllNodes())
+			for (USCS_Node* Node : SCS->GetAllNodes())
 			{
 				ActorComponentTempatesOwnedByClass.RemoveSwap(Node->ComponentTemplate);
 			}
 		}
-		if (auto IH = BPGC->GetInheritableComponentHandler())
+		if (UInheritableComponentHandler* IH = BPGC->GetInheritableComponentHandler())
 		{
 			TArray<UActorComponent*> AllTemplates;
 			IH->GetAllTemplates(AllTemplates);
@@ -1746,7 +1746,7 @@ void FEmitDefaultValueHelper::GenerateCustomDynamicClassInitialization(FEmitterL
 
 		auto CreateAndInitializeClassSubobjects = [&](bool bCreate, bool bInitialize)
 		{
-			for (auto ComponentTemplate : ActorComponentTempatesOwnedByClass)
+			for (UActorComponent* ComponentTemplate : ActorComponentTempatesOwnedByClass)
 			{
 				if (ComponentTemplate)
 				{
@@ -1754,7 +1754,7 @@ void FEmitDefaultValueHelper::GenerateCustomDynamicClassInitialization(FEmitterL
 				}
 			}
 
-			for (auto TimelineTemplate : BPGC->Timelines)
+			for (UTimelineTemplate* TimelineTemplate : BPGC->Timelines)
 			{
 				if (TimelineTemplate)
 				{
@@ -1762,7 +1762,7 @@ void FEmitDefaultValueHelper::GenerateCustomDynamicClassInitialization(FEmitterL
 				}
 			}
 
-			for (auto DynamicBindingObject : BPGC->DynamicBindingObjects)
+			for (UDynamicBlueprintBinding* DynamicBindingObject : BPGC->DynamicBindingObjects)
 			{
 				if (DynamicBindingObject)
 				{
@@ -1823,7 +1823,7 @@ void FEmitDefaultValueHelper::GenerateConstructor(FEmitterLocalContext& Context)
 
 		// Subobjects that must be fixed after serialization
 		TArray<FDefaultSubobjectData> SubobjectsToInit;
-		TArray<FNonativeComponentData> ComponentsToInit;
+		TArray<FNonNativeComponentData> ComponentsToInit;
 
 		{
 			Context.AddLine(TEXT(""));
@@ -1835,7 +1835,7 @@ void FEmitDefaultValueHelper::GenerateConstructor(FEmitterLocalContext& Context)
 			// @TODO - We can probably make this faster by generating code to directly index through the DSO array instead (i.e. in place of HandleInstancedSubobject which will generate a lookup call per DSO).
 			TArray<UObject*> NativeDefaultObjectSubobjects;
 			BPGC->GetDefaultObjectSubobjects(NativeDefaultObjectSubobjects);
-			for (auto DSO : NativeDefaultObjectSubobjects)
+			for (UObject* DSO : NativeDefaultObjectSubobjects)
 			{
 				if (DSO && DSO->GetClass()->HasAnyClassFlags(CLASS_DefaultToInstanced))
 				{
@@ -1863,7 +1863,7 @@ void FEmitDefaultValueHelper::GenerateConstructor(FEmitterLocalContext& Context)
 			}
 
 			// Emit the code to initialize all instanced default subobjects now referenced by a local variable.
-			for (auto& DSOEntry : SubobjectsToInit)
+			for (FDefaultSubobjectData& DSOEntry : SubobjectsToInit)
 			{
 				DSOEntry.EmitPropertyInitialization(Context);
 			}
@@ -1925,7 +1925,7 @@ void FEmitDefaultValueHelper::GenerateConstructor(FEmitterLocalContext& Context)
 						{
 							if (Node)
 							{
-								const bool bNodeWasProcessed = nullptr != ComponentsToInit.FindByPredicate([=](const FNonativeComponentData& InData) { return Node == InData.SCSNode; });
+								const bool bNodeWasProcessed = nullptr != ComponentsToInit.FindByPredicate([=](const FNonNativeComponentData& InData) { return Node == InData.SCSNode; });
 								if (!bNodeWasProcessed)
 								{
 									HandleNonNativeComponent(Context, Node, HandledProperties, NativeCreatedComponentProperties, nullptr, ComponentsToInit, true);
@@ -1936,14 +1936,14 @@ void FEmitDefaultValueHelper::GenerateConstructor(FEmitterLocalContext& Context)
 					}
 				}
 
-				for (auto& ComponentToInit : ComponentsToInit)
+				for (FNonNativeComponentData& ComponentToInit : ComponentsToInit)
 				{
 					ComponentToInit.EmitPropertyInitialization(Context);
 				}
 			}
 
 			// Collect all anim node properties
-			for (auto Property : TFieldRange<const UProperty>(BPGC))
+			for (const UProperty* Property : TFieldRange<const UProperty>(BPGC))
 			{
 				if (!HandledProperties.Contains(Property))
 				{
@@ -1961,7 +1961,7 @@ void FEmitDefaultValueHelper::GenerateConstructor(FEmitterLocalContext& Context)
 			}
 
 			// Generate ctor init code for generated Blueprint class property values that may differ from parent class defaults (or that otherwise belong to the generated Blueprint class).
-			for (auto Property : TFieldRange<const UProperty>(BPGC))
+			for (const UProperty* Property : TFieldRange<const UProperty>(BPGC))
 			{
 				if (!HandledProperties.Contains(Property))
 				{
@@ -2001,7 +2001,7 @@ void FEmitDefaultValueHelper::GenerateConstructor(FEmitterLocalContext& Context)
 	Context.AddLine(TEXT("{"));
 	Context.IncreaseIndent();
 	Context.AddLine(FString::Printf(TEXT("Super::%s(OuterInstanceGraph);"), GET_FUNCTION_NAME_STRING_CHECKED(UObject, PostLoadSubobjects)));
-	for (auto& ComponentToFix : NativeCreatedComponentProperties)
+	for (FString& ComponentToFix : NativeCreatedComponentProperties)
 	{
 		Context.AddLine(FString::Printf(TEXT("if(%s)"), *ComponentToFix));
 		Context.AddLine(TEXT("{"));
@@ -2082,8 +2082,8 @@ FString FEmitDefaultValueHelper::HandleClassSubobject(FEmitterLocalContext& Cont
 		
 		if (ensure(!LocalNativeName.IsEmpty()))
 		{
-			auto CDO = Object->GetClass()->GetDefaultObject(false);
-			for (auto Property : TFieldRange<const UProperty>(Object->GetClass()))
+			UObject* CDO = Object->GetClass()->GetDefaultObject(false);
+			for (const UProperty* Property : TFieldRange<const UProperty>(Object->GetClass()))
 			{
 				OuterGenerate(Context, Property, LocalNativeName
 					, reinterpret_cast<const uint8*>(Object)
@@ -2136,8 +2136,8 @@ FString FEmitDefaultValueHelper::HandleInstancedSubobject(FEmitterLocalContext&
 		}
 	}
 
-	auto BPGC = Context.GetCurrentlyGeneratedClass();
-	auto CDO = BPGC ? BPGC->GetDefaultObject(false) : nullptr;
+	UClass* BPGC = Context.GetCurrentlyGeneratedClass();
+	UObject* CDO = BPGC ? BPGC->GetDefaultObject(false) : nullptr;
 
 	FString OuterStr;
 	if (ensure(CDO) && (CDO == Object->GetOuter()))
