Fix static analysis warning pointing out both too many parameters to a function and a crash!
#rnx

diff --git a/Engine/Source/Developer/BlueprintCompilerCppBackend/Private/BlueprintCompilerCppBackendValueHelper.cpp b/Engine/Source/Developer/BlueprintCompilerCppBackend/Private/BlueprintCompilerCppBackendValueHelper.cpp
--- a/Engine/Source/Developer/BlueprintCompilerCppBackend/Private/BlueprintCompilerCppBackendValueHelper.cpp
+++ b/Engine/Source/Developer/BlueprintCompilerCppBackend/Private/BlueprintCompilerCppBackendValueHelper.cpp
@@ -1783,7 +1783,7 @@ void FEmitDefaultValueHelper::GenerateCustomDynamicClassInitialization(FEmitterL
 			}
 			else
 			{
-				Context.AddLine(FString::Printf(TEXT("InDynamicClass->ComponentClassOverrides.Emplace(MakeTuple(FName(\"%s\"), nullptr));"), *Override.ComponentName.ToString(), *Override.ComponentClass->GetName()));
+				Context.AddLine(FString::Printf(TEXT("InDynamicClass->ComponentClassOverrides.Emplace(MakeTuple(FName(\"%s\"), nullptr));"), *Override.ComponentName.ToString()));
 			}
 		}
 
