Fix nativized interfaces not compiling
#rnx

diff --git a/Engine/Source/Developer/BlueprintCompilerCppBackend/Private/BlueprintCompilerCppBackendBase.cpp b/Engine/Source/Developer/BlueprintCompilerCppBackend/Private/BlueprintCompilerCppBackendBase.cpp
--- a/Engine/Source/Developer/BlueprintCompilerCppBackend/Private/BlueprintCompilerCppBackendBase.cpp
+++ b/Engine/Source/Developer/BlueprintCompilerCppBackend/Private/BlueprintCompilerCppBackendBase.cpp
@@ -423,7 +423,11 @@ FString FBlueprintCompilerCppBackendBase::GenerateCodeFromClass(UClass* SourceCl
 
 		TSharedPtr<FGatherConvertedClassDependencies> ParentDependencies;
 		// Emit function declarations and definitions (writes to header and body simultaneously)
-		if (!bIsInterface)
+		if (bIsInterface)
+		{
+			EmitterContext.Header.AddLine(TEXT("static void __CustomDynamicClassInitialization(UDynamicClass* InDynamicClass) {}"));
+		}
+		else
 		{
 			UBlueprintGeneratedClass* BPGC = CastChecked<UBlueprintGeneratedClass>(EmitterContext.GetCurrentlyGeneratedClass());
 			UBlueprintGeneratedClass* ParentBPGC = Cast<UBlueprintGeneratedClass>(BPGC->GetSuperClass());
