Generate __CustomDynamicClassInitialization in the U class not the I class for nativized interfaces
#rnx

diff --git a/Engine/Source/Developer/BlueprintCompilerCppBackend/Private/BlueprintCompilerCppBackendBase.cpp b/Engine/Source/Developer/BlueprintCompilerCppBackend/Private/BlueprintCompilerCppBackendBase.cpp
--- a/Engine/Source/Developer/BlueprintCompilerCppBackend/Private/BlueprintCompilerCppBackendBase.cpp
+++ b/Engine/Source/Developer/BlueprintCompilerCppBackend/Private/BlueprintCompilerCppBackendBase.cpp
@@ -357,6 +357,7 @@ FString FBlueprintCompilerCppBackendBase::GenerateCodeFromClass(UClass* SourceCl
 			EmitterContext.Header.AddLine(TEXT("{"));
 			EmitterContext.Header.IncreaseIndent();
 			EmitterContext.Header.AddLine(TEXT("GENERATED_BODY()"));
+			EmitterContext.Header.AddLine(TEXT("static void __CustomDynamicClassInitialization(UDynamicClass* InDynamicClass) {}"));
 			EmitterContext.Header.DecreaseIndent();
 			EmitterContext.Header.AddLine(TEXT("};"));
 			EmitterContext.Header.AddLine(FString::Printf(TEXT("class %s"), *CppClassName));
@@ -423,11 +424,7 @@ FString FBlueprintCompilerCppBackendBase::GenerateCodeFromClass(UClass* SourceCl
 
 		TSharedPtr<FGatherConvertedClassDependencies> ParentDependencies;
 		// Emit function declarations and definitions (writes to header and body simultaneously)
-		if (bIsInterface)
-		{
-			EmitterContext.Header.AddLine(TEXT("static void __CustomDynamicClassInitialization(UDynamicClass* InDynamicClass) {}"));
-		}
-		else
+		if (!bIsInterface)
 		{
 			UBlueprintGeneratedClass* BPGC = CastChecked<UBlueprintGeneratedClass>(EmitterContext.GetCurrentlyGeneratedClass());
 			UBlueprintGeneratedClass* ParentBPGC = Cast<UBlueprintGeneratedClass>(BPGC->GetSuperClass());
