Fix typo in blueprint nativization with TMap entry default Value assignment.

diff --git a/Engine/Source/Developer/BlueprintCompilerCppBackend/Private/BlueprintCompilerCppBackendValueHelper.cpp b/Engine/Source/Developer/BlueprintCompilerCppBackend/Private/BlueprintCompilerCppBackendValueHelper.cpp
--- a/Engine/Source/Developer/BlueprintCompilerCppBackend/Private/BlueprintCompilerCppBackendValueHelper.cpp
+++ b/Engine/Source/Developer/BlueprintCompilerCppBackend/Private/BlueprintCompilerCppBackendValueHelper.cpp
@@ -463,7 +463,7 @@ void FEmitDefaultValueHelper::InnerGenerate(FEmitterLocalContext& Context, const
 						if (EStructConstructionType::Custom == ValueConstruction)
 						{
 							FString ValueStr;
-							bValueComplete = OneLineConstruction(Context, MapProperty->ValueProp, ScriptMapHelper.GetKeyPtr(Index), ValueStr, /*bGenerateEmptyStructConstructor=*/ false);
+							bValueComplete = OneLineConstruction(Context, MapProperty->ValueProp, ScriptMapHelper.GetValuePtr(Index), ValueStr, /*bGenerateEmptyStructConstructor=*/ false);
 							if (!ValueStr.IsEmpty())
 							{
 								Context.AddLine(FString::Printf(TEXT("%s = %s;"), *ValuePath, *ValueStr));
