Fix spelling error and don't create a few FStrings where a const TCHAR* would suffice.
#rnx

diff --git a/Engine/Source/Developer/BlueprintCompilerCppBackend/Private/BlueprintCompilerCppBackendValueHelper.cpp b/Engine/Source/Developer/BlueprintCompilerCppBackend/Private/BlueprintCompilerCppBackendValueHelper.cpp
--- a/Engine/Source/Developer/BlueprintCompilerCppBackend/Private/BlueprintCompilerCppBackendValueHelper.cpp
+++ b/Engine/Source/Developer/BlueprintCompilerCppBackend/Private/BlueprintCompilerCppBackendValueHelper.cpp
@@ -81,8 +81,8 @@ void FEmitDefaultValueHelper::OuterGenerate(FEmitterLocalContext& Context
 			if (PropertyOwnerAsBPGC && !Context.Dependencies.WillClassBeConverted(PropertyOwnerAsBPGC))
 			{
 				ensure(EPropertyAccessOperator::None != AccessOperator);
-				const FString OperatorStr = (EPropertyAccessOperator::Dot == AccessOperator) ? TEXT("&") : TEXT("");
-				const FString ContainerStr = (EPropertyAccessOperator::None == AccessOperator) ? TEXT("this") : FString::Printf(TEXT("%s(%s)"), *OperatorStr, *OuterPath);
+				const TCHAR* OperatorStr = (EPropertyAccessOperator::Dot == AccessOperator) ? TEXT("&") : TEXT("");
+				const FString ContainerStr = (EPropertyAccessOperator::None == AccessOperator) ? TEXT("this") : FString::Printf(TEXT("%s(%s)"), OperatorStr, *OuterPath);
 
 				PathToMember = FString::Printf(TEXT("FUnconvertedWrapper__%s(%s).GetRef__%s()"), *FEmitHelper::GetCppName(PropertyOwnerAsBPGC), *ContainerStr
 					, *UnicodeToCPPIdentifier(Property->GetName(), false, nullptr));
@@ -92,8 +92,8 @@ void FEmitDefaultValueHelper::OuterGenerate(FEmitterLocalContext& Context
 			{
 				const UBoolProperty* BoolProperty = Cast<const UBoolProperty>(Property);
 				const bool bBietfield = BoolProperty && !BoolProperty->IsNativeBool();
-				const FString OperatorStr = (EPropertyAccessOperator::Dot == AccessOperator) ? TEXT("&") : TEXT("");
-				const FString ContainerStr = (EPropertyAccessOperator::None == AccessOperator) ? TEXT("this") : OuterPath;
+				const TCHAR* OperatorStr = (EPropertyAccessOperator::Dot == AccessOperator) ? TEXT("&") : TEXT("");
+				const TCHAR* ContainerStr = (EPropertyAccessOperator::None == AccessOperator) ? TEXT("this") : *OuterPath;
 				if (bBietfield)
 				{
 					const FString PropertyLocalName = FEmitHelper::GenerateGetPropertyByName(Context, Property);
@@ -101,8 +101,8 @@ void FEmitDefaultValueHelper::OuterGenerate(FEmitterLocalContext& Context
 					Context.AddLine(FString::Printf(TEXT("(((UBoolProperty*)%s)->%s(%s(%s), %s, %d));")
 						, *PropertyLocalName
 						, GET_FUNCTION_NAME_STRING_CHECKED(UBoolProperty, SetPropertyValue_InContainer)
-						, *OperatorStr
-						, *ContainerStr
+						, OperatorStr
+						, ContainerStr
 						, *ValueStr
 						, ArrayIndex));
 					continue;
@@ -125,11 +125,11 @@ void FEmitDefaultValueHelper::OuterGenerate(FEmitterLocalContext& Context
 			}
 			else
 			{
-				const FString AccessOperatorStr = (EPropertyAccessOperator::None == AccessOperator) ? TEXT("")
+				const TCHAR* AccessOperatorStr = (EPropertyAccessOperator::None == AccessOperator) ? TEXT("")
 					: ((EPropertyAccessOperator::Pointer == AccessOperator) ? TEXT("->") : TEXT("."));
 				const bool bStaticArray = (Property->ArrayDim > 1);
 				const FString ArrayPost = bStaticArray ? FString::Printf(TEXT("[%d]"), ArrayIndex) : TEXT("");
-				PathToMember = FString::Printf(TEXT("%s%s%s%s"), *OuterPath, *AccessOperatorStr, *FEmitHelper::GetCppName(Property), *ArrayPost);
+				PathToMember = FString::Printf(TEXT("%s%s%s%s"), *OuterPath, AccessOperatorStr, *FEmitHelper::GetCppName(Property), *ArrayPost);
 			}
 
 			const uint8* ValuePtr = Property->ContainerPtrToValuePtr<uint8>(DataContainer, ArrayIndex);
