Fix a BP nativization issue in which latent node exec pins occasionally do not fire in the converted class.

diff --git a/Engine/Source/Developer/BlueprintCompilerCppBackend/Private/BlueprintCompilerCppBackend.cpp b/Engine/Source/Developer/BlueprintCompilerCppBackend/Private/BlueprintCompilerCppBackend.cpp
--- a/Engine/Source/Developer/BlueprintCompilerCppBackend/Private/BlueprintCompilerCppBackend.cpp
+++ b/Engine/Source/Developer/BlueprintCompilerCppBackend/Private/BlueprintCompilerCppBackend.cpp
@@ -1142,7 +1142,7 @@ FString FBlueprintCompilerCppBackend::LatentFunctionInfoTermToText(FEmitterLocal
 	const int32 TargetStateIndex = StateMapPerFunction[0].StatementToStateIndex(TargetLabel);
 	const int32 LinkageTermStartIdx = LiteralTermParams.CustomValue.Find(FixupTermName);
 	check(LinkageTermStartIdx != INDEX_NONE);
-	LiteralTermParams.CustomValue = LiteralTermParams.CustomValue.Replace(TEXT("-1"), *FString::FromInt(TargetStateIndex));
+	LiteralTermParams.CustomValue.ReplaceInline(*FString::Printf(TEXT("%s=-1"), *FixupTermName), *FString::Printf(TEXT("%s=%d"), *FixupTermName, TargetStateIndex));
 
 	int32* ExecutionGroupPtr = UberGraphStatementToExecutionGroup.Find(TargetLabel);
 	if (ExecutionGroupPtr && UberGraphContext)
