Fix for a BP nativization issue related to new method of widget animation event binding.

#rnx

diff --git a/Engine/Source/Developer/BlueprintCompilerCppBackend/Private/BlueprintCompilerCppBackendUtils.cpp b/Engine/Source/Developer/BlueprintCompilerCppBackend/Private/BlueprintCompilerCppBackendUtils.cpp
--- a/Engine/Source/Developer/BlueprintCompilerCppBackend/Private/BlueprintCompilerCppBackendUtils.cpp
+++ b/Engine/Source/Developer/BlueprintCompilerCppBackend/Private/BlueprintCompilerCppBackendUtils.cpp
@@ -163,10 +163,21 @@ FString FEmitterLocalContext::FindGloballyMappedObject(const UObject* Object, co
 						, *MappedOwner, *Field->GetName());
 				}
 
+				FString FieldName = Field->GetName();
+				if (FieldOwnerStruct->IsA<UFunction>())
+				{
+					// Function-owned fields (e.g. params) don't currently direct UHT to override the nativized field name if the owning class will be converted.
+					const UBlueprintGeneratedClass* BPGC = Cast<UBlueprintGeneratedClass>(FieldOwnerStruct->GetOwnerClass());
+					if (BPGC && Dependencies.WillClassBeConverted(BPGC))
+					{
+						FieldName = FEmitHelper::GetCppName(Field);
+					}
+				}
+
 				return FString::Printf(TEXT("FindFieldChecked<%s>(%s, TEXT(\"%s\"))")
 					, *FEmitHelper::GetCppName(Field->GetClass())
 					, *MappedOwner
-					, *Field->GetName());
+					, *FieldName);
 			}
 		}
 	}
