Fix an invalid cast in nativized Blueprint code for converted user-defined enum fields.

- Mirrored from //UE4/Dev-Framework (10861958).

diff --git a/Engine/Source/Developer/BlueprintCompilerCppBackend/Private/BlueprintCompilerCppBackendUtils.cpp b/Engine/Source/Developer/BlueprintCompilerCppBackend/Private/BlueprintCompilerCppBackendUtils.cpp
--- a/Engine/Source/Developer/BlueprintCompilerCppBackend/Private/BlueprintCompilerCppBackendUtils.cpp
+++ b/Engine/Source/Developer/BlueprintCompilerCppBackend/Private/BlueprintCompilerCppBackendUtils.cpp
@@ -16,6 +16,7 @@
 #include "UObject/UObjectHash.h"
 #include "KismetCompiler.h"
 #include "Misc/DefaultValueHelper.h"
+#include "Blueprint/BlueprintSupport.h"
 #include "BlueprintCompilerCppBackend.h"
 #include "Animation/AnimBlueprint.h"
 
@@ -175,8 +176,15 @@ FString FEmitterLocalContext::FindGloballyMappedObject(const UObject* Object, co
 					}
 				}
 
+				// Some field types may be replaced after conversion (e.g. converted user-defined enum types).
+				const UClass* FieldClass = Field->GetClass();
+				if (const UClass* ReplacedClass = IBlueprintNativeCodeGenCore::Get()->FindReplacedClassForObject(Field, NativizationOptions))
+				{
+					FieldClass = ReplacedClass;
+				}
+
 				return FString::Printf(TEXT("FindFieldChecked<%s>(%s, TEXT(\"%s\"))")
-					, *FEmitHelper::GetCppName(Field->GetClass())
+					, *FEmitHelper::GetCppName(FieldClass)
 					, *MappedOwner
 					, *FieldName);
 			}
