Fix a typo in code that's emitted for properties that reference explicitly-mapped 'noexport' struct fields in nativized Blueprint code.

- Mirrored from //UE4/Dev-Framework (CL# 10864850).

diff --git a/Engine/Source/Developer/BlueprintCompilerCppBackend/Private/BlueprintCompilerCppBackendUtils.cpp b/Engine/Source/Developer/BlueprintCompilerCppBackend/Private/BlueprintCompilerCppBackendUtils.cpp
--- a/Engine/Source/Developer/BlueprintCompilerCppBackend/Private/BlueprintCompilerCppBackendUtils.cpp
+++ b/Engine/Source/Developer/BlueprintCompilerCppBackend/Private/BlueprintCompilerCppBackendUtils.cpp
@@ -2318,7 +2318,7 @@ void FScopeBlock::TrackLocalAccessorDecl(const FProperty* Property)
 }
 
 #define MAP_BASE_STRUCTURE_ACCESS(x) \
-	BaseStructureAccessorsMap.Add(x, TEXT("#x"))
+	BaseStructureAccessorsMap.Add(x, TEXT(#x))
 
 struct FStructAccessHelper_StaticData
 {
