PR #6492: Allow UCLASS metadata to specify a CustomThunkTemplates class so that CustomThunk nodes can be nativized (Contributed by [REDACTED_PERSON])

diff --git a/Engine/Source/Developer/BlueprintCompilerCppBackend/Private/BlueprintCompilerCppBackend.cpp b/Engine/Source/Developer/BlueprintCompilerCppBackend/Private/BlueprintCompilerCppBackend.cpp
--- a/Engine/Source/Developer/BlueprintCompilerCppBackend/Private/BlueprintCompilerCppBackend.cpp
+++ b/Engine/Source/Developer/BlueprintCompilerCppBackend/Private/BlueprintCompilerCppBackend.cpp
@@ -850,7 +850,28 @@ FString FBlueprintCompilerCppBackend::EmitCallStatmentInner(FEmitterLocalContext
 		else if (bStaticCall)
 		{
 			UClass* OwnerClass = Statement.FunctionToCall->GetOuterUClass();
-			Result += bIsCustomThunk ? TEXT("FCustomThunkTemplates::") : FString::Printf(TEXT("%s::"), *FEmitHelper::GetCppName(OwnerClass));
+			if (bIsCustomThunk)
+			{
+				// Check class metadata for a specified CustomThunkTemplates override
+				FName CustomThunkTemplatesMetadataKey = FName(TEXT("CustomThunkTemplates"));
+				FString CustomThunkTemplatesName;
+				if (OwnerClass->HasMetaDataHierarchical(CustomThunkTemplatesMetadataKey))
+				{
+					OwnerClass->GetStringMetaDataHierarchical(CustomThunkTemplatesMetadataKey, &CustomThunkTemplatesName);
+				}
+				else
+				{
+					// Fall back to the engine-supplied templates
+					CustomThunkTemplatesName = TEXT("FCustomThunkTemplates");
+				}
+
+				Result += CustomThunkTemplatesName + TEXT("::");
+			}
+			else
+			{
+				Result += FString::Printf(TEXT("%s::"), *FEmitHelper::GetCppName(OwnerClass));
+			}
+
 		}
 		else if (bCallOnDifferentObject) //@TODO: Badness, could be a self reference wired to another instance!
 		{
