Fix blind reinstancing of natively-constructed nested subobjects during class-owned subobject initialization in nativized Blueprint C++ code.

diff --git a/Engine/Source/Developer/BlueprintCompilerCppBackend/Private/BlueprintCompilerCppBackendValueHelper.cpp b/Engine/Source/Developer/BlueprintCompilerCppBackend/Private/BlueprintCompilerCppBackendValueHelper.cpp
--- a/Engine/Source/Developer/BlueprintCompilerCppBackend/Private/BlueprintCompilerCppBackendValueHelper.cpp
+++ b/Engine/Source/Developer/BlueprintCompilerCppBackend/Private/BlueprintCompilerCppBackendValueHelper.cpp
@@ -727,97 +727,6 @@
 	return false;
 }
 
-FString FEmitDefaultValueHelper::HandleSpecialTypes(FEmitterLocalContext& Context, const FProperty* Property, const uint8* ValuePtr)
-{
-	auto HandleObjectValueLambda = [&Context, Property, ValuePtr](UObject* Object, UClass* Class) -> FString
-	{
-		if (Object)
-		{
-			const bool bIsDefaultSubobject = Object->IsDefaultSubobject() && Object->HasAllFlags(RF_DefaultSubObject);
-			const bool bIsInstancedReference = Property->HasAnyPropertyFlags(CPF_InstancedReference);
-
-			UClass* ObjectClassToUse = Context.GetFirstNativeOrConvertedClass(Class);
-			{
-				const FString MappedObject = Context.FindGloballyMappedObject(Object, ObjectClassToUse);
-				if (!MappedObject.IsEmpty())
-				{
-					return MappedObject;
-				}
-			}
-
-			UClass* BPGC = Context.GetCurrentlyGeneratedClass();
-
-			UChildActorComponent* OuterCAC = Cast<UChildActorComponent>(Object->GetOuter());
-			const bool bObjectIsCACTemplate = OuterCAC && OuterCAC->IsIn(BPGC) && OuterCAC->GetChildActorTemplate() == Object;
-
-			const bool bCreatingSubObjectsOfClass = (Context.CurrentCodeType == FEmitterLocalContext::EGeneratedCodeType::SubobjectsOfClass);
-			{
-				UObject* CDO = BPGC ? BPGC->GetDefaultObject(false) : nullptr;
-				if (BPGC && Object && CDO && Object->IsIn(BPGC) && !Object->IsIn(CDO) && bCreatingSubObjectsOfClass)
-				{
-					return HandleClassSubobject(Context, Object, FEmitterLocalContext::EClassSubobjectList::MiscConvertedSubobjects, true, true, bObjectIsCACTemplate);
-				}
-			}
-
-			if (!bCreatingSubObjectsOfClass && bIsInstancedReference)
-			{
-				// Emit ctor code to create the instance only if it's not a default subobject; otherwise, just assign the reference value to a local variable for initialization.
-				// Note that we also skip the editor-only check if it's a default subobject. In that case, the instance will either have already been created with CreateDefaultSubobject(),
-				// or creation will have been skipped (e.g. CreateEditorOnlyDefaultSubobject()). We check the pointer for NULL before assigning default value overrides in the generated ctor.
-				const FString MappedObject = HandleInstancedSubobject(Context, Object, /* bCreateInstance = */ !bIsDefaultSubobject, /* bSkipEditorOnlyCheck = */ bIsDefaultSubobject);
-
-				// We should always find a mapping in this case.
-				if (ensure(!MappedObject.IsEmpty()))
-				{
-					return MappedObject;
-				}
-			}
-
-			if (!bCreatingSubObjectsOfClass && bObjectIsCACTemplate)
-			{
-				Context.TemplateFromSubobjectsOfClass.AddUnique(Object);
-				const FString MappedObject = Context.FindGloballyMappedObject(Object, ObjectClassToUse);
-				if (!MappedObject.IsEmpty())
-				{
-					return MappedObject;
-				}
-			}
-		}
-		else
-		{
-			// Emit valid representation for a null object.
-			return Context.ExportTextItem(Property, ValuePtr);
-		}
-
-		return FString();
-	};
-
-	FString Result;
-
-	if (const FObjectProperty* ObjectProperty = CastField<FObjectProperty>(Property))
-	{
-		Result = HandleObjectValueLambda(ObjectProperty->GetPropertyValue(ValuePtr), ObjectProperty->PropertyClass);
-	}
-	else if (const FWeakObjectProperty* WeakObjectProperty = CastField<FWeakObjectProperty>(Property))
-	{
-		Result = HandleObjectValueLambda(WeakObjectProperty->GetObjectPropertyValue(ValuePtr), WeakObjectProperty->PropertyClass);
-	}
-	else if (const FInterfaceProperty* InterfaceProperty = CastField<FInterfaceProperty>(Property))
-	{
-		Result = HandleObjectValueLambda(InterfaceProperty->GetPropertyValue(ValuePtr).GetObject(), InterfaceProperty->InterfaceClass);
-	}
-	else if (const FStructProperty* StructProperty = CastField<FStructProperty>(Property))
-	{
-		FString StructConstructor;
-		if (SpecialStructureConstructor(StructProperty->Struct, ValuePtr, &StructConstructor))
-		{
-			Result = StructConstructor;
-		}
-	}
-
-	return Result;
-}
-
 struct FDefaultSubobjectData
 {
 	UObject* Object;
@@ -1027,6 +936,116 @@
 	}
 };
 
+FString FEmitDefaultValueHelper::HandleSpecialTypes(FEmitterLocalContext& Context, const FProperty* Property, const uint8* ValuePtr)
+{
+	auto HandleObjectValueLambda = [&Context, Property, ValuePtr](UObject* Object, UClass* Class) -> FString
+	{
+		if (Object)
+		{
+			const bool bIsDefaultSubobject = Object->IsDefaultSubobject() && Object->HasAllFlags(RF_DefaultSubObject);
+			const bool bIsInstancedReference = Property->HasAnyPropertyFlags(CPF_InstancedReference);
+
+			UClass* ObjectClassToUse = Context.GetFirstNativeOrConvertedClass(Class);
+			{
+				const FString MappedObject = Context.FindGloballyMappedObject(Object, ObjectClassToUse);
+				if (!MappedObject.IsEmpty())
+				{
+					return MappedObject;
+				}
+			}
+
+			UClass* BPGC = Context.GetCurrentlyGeneratedClass();
+
+			UChildActorComponent* OuterCAC = Cast<UChildActorComponent>(Object->GetOuter());
+			const bool bObjectIsCACTemplate = OuterCAC && OuterCAC->IsIn(BPGC) && OuterCAC->GetChildActorTemplate() == Object;
+
+			const bool bCreatingSubObjectsOfClass = (Context.CurrentCodeType == FEmitterLocalContext::EGeneratedCodeType::SubobjectsOfClass);
+			{
+				UObject* CDO = BPGC ? BPGC->GetDefaultObject(false) : nullptr;
+				if (BPGC && Object && CDO && Object->IsIn(BPGC) && !Object->IsIn(CDO) && bCreatingSubObjectsOfClass)
+				{
+					if (Object->GetOuter() == BPGC || bObjectIsCACTemplate)
+					{
+						return HandleClassSubobject(Context, Object, FEmitterLocalContext::EClassSubobjectList::MiscConvertedSubobjects, true, true, bObjectIsCACTemplate);
+					}
+					else
+					{
+						// Treat initialization of nested objects owned by class subobjects like "normal" instanced subobjects, but disable scoped initialization because class subobjects can also reference nested DSOs.
+						FDefaultSubobjectData ClassSubobjectNestedDSOData;
+						ClassSubobjectNestedDSOData.bAddLocalScope = false;
+						const FString MappedObject = HandleInstancedSubobject(Context, Object, !bIsDefaultSubobject, bIsDefaultSubobject, &ClassSubobjectNestedDSOData);
+						
+						// We should always find a mapping in this case.
+						if (ensure(!MappedObject.IsEmpty()))
+						{
+							// Since we are using our own subobject data context to avoid scoping, we must explicitly emit the initialization code.
+							ClassSubobjectNestedDSOData.EmitPropertyInitialization(Context);
+
+							return MappedObject;
+						}
+					}
+				}
+			}
+
+			if (!bCreatingSubObjectsOfClass && bIsInstancedReference)
+			{
+				// Emit ctor code to create the instance only if it's not a default subobject; otherwise, just assign the reference value to a local variable for initialization.
+				// Note that we also skip the editor-only check if it's a default subobject. In that case, the instance will either have already been created with CreateDefaultSubobject(),
+				// or creation will have been skipped (e.g. CreateEditorOnlyDefaultSubobject()). We check the pointer for NULL before assigning default value overrides in the generated ctor.
+				const FString MappedObject = HandleInstancedSubobject(Context, Object, /* bCreateInstance = */ !bIsDefaultSubobject, /* bSkipEditorOnlyCheck = */ bIsDefaultSubobject);
+
+				// We should always find a mapping in this case.
+				if (ensure(!MappedObject.IsEmpty()))
+				{
+					return MappedObject;
+				}
+			}
+
+			if (!bCreatingSubObjectsOfClass && bObjectIsCACTemplate)
+			{
+				Context.TemplateFromSubobjectsOfClass.AddUnique(Object);
+				const FString MappedObject = Context.FindGloballyMappedObject(Object, ObjectClassToUse);
+				if (!MappedObject.IsEmpty())
+				{
+					return MappedObject;
+				}
+			}
+		}
+		else
+		{
+			// Emit valid representation for a null object.
+			return Context.ExportTextItem(Property, ValuePtr);
+		}
+
+		return FString();
+	};
+
+	FString Result;
+
+	if (const FObjectProperty* ObjectProperty = CastField<FObjectProperty>(Property))
+	{
+		Result = HandleObjectValueLambda(ObjectProperty->GetPropertyValue(ValuePtr), ObjectProperty->PropertyClass);
+	}
+	else if (const FWeakObjectProperty* WeakObjectProperty = CastField<FWeakObjectProperty>(Property))
+	{
+		Result = HandleObjectValueLambda(WeakObjectProperty->GetObjectPropertyValue(ValuePtr), WeakObjectProperty->PropertyClass);
+	}
+	else if (const FInterfaceProperty* InterfaceProperty = CastField<FInterfaceProperty>(Property))
+	{
+		Result = HandleObjectValueLambda(InterfaceProperty->GetPropertyValue(ValuePtr).GetObject(), InterfaceProperty->InterfaceClass);
+	}
+	else if (const FStructProperty* StructProperty = CastField<FStructProperty>(Property))
+	{
+		FString StructConstructor;
+		if (SpecialStructureConstructor(StructProperty->Struct, ValuePtr, &StructConstructor))
+		{
+			Result = StructConstructor;
+		}
+	}
+
+	return Result;
+}
+
 FString FEmitDefaultValueHelper::HandleNonNativeComponent(FEmitterLocalContext& Context, const USCS_Node* Node
 	, TSet<const FProperty*>& OutHandledProperties, TArray<FString>& NativeCreatedComponentProperties
 	, const USCS_Node* ParentNode, TArray<FNonNativeComponentData>& ComponentsToInit
