Fixing BP nativization crash when cooking

diff --git a/Engine/Source/Developer/BlueprintCompilerCppBackend/Private/BlueprintCompilerCppBackendUtils.cpp b/Engine/Source/Developer/BlueprintCompilerCppBackend/Private/BlueprintCompilerCppBackendUtils.cpp
--- a/Engine/Source/Developer/BlueprintCompilerCppBackend/Private/BlueprintCompilerCppBackendUtils.cpp
+++ b/Engine/Source/Developer/BlueprintCompilerCppBackend/Private/BlueprintCompilerCppBackendUtils.cpp
@@ -463,7 +463,7 @@ FString FEmitHelper::GetCppName(FFieldVariant Field, bool bUInterface, bool bFor
 				, bUInterface ? TEXT("U") : TEXT("I")
 				, *AsClass->GetName());
 		}
-		const UStruct* AsStruct = Field.Get<UScriptStruct>();
+		const UStruct* AsStruct = Field.Get<UStruct>();
 		check(AsStruct);
 		if (AsStruct->IsNative())
 		{
@@ -623,24 +623,27 @@ bool FEmitHelper::MetaDataCanBeNative(const FName MetaDataName, FFieldVariant Fi
 FString FEmitHelper::HandleMetaData(FFieldVariant Field, bool AddCategory, const TArray<FString>* AdditinalMetaData)
 {
 	FString MetaDataStr;
-
 	TMap<FName, FString> ValuesMap;
-	
-	if (Field.IsUObject())
+	TArray<FString> MetaDataStrings;
+
+	if (Field.IsValid())
 	{
-		UPackage* Package = Field.IsValid() ? Field.GetOutermost() : nullptr;
-		const UMetaData* MetaData = Package ? Package->GetMetaData() : nullptr;
-		const TMap<FName, FString>* ValuesMapPtr = MetaData ? MetaData->ObjectMetaDataMap.Find(Field.ToUObject()) : nullptr;
-		if (ValuesMapPtr)
+		if (Field.IsUObject())
 		{
-			ValuesMap = *ValuesMapPtr;
+			UPackage* Package = Field.IsValid() ? Field.GetOutermost() : nullptr;
+			const UMetaData* MetaData = Package ? Package->GetMetaData() : nullptr;
+			const TMap<FName, FString>* ValuesMapPtr = MetaData ? MetaData->ObjectMetaDataMap.Find(Field.ToUObject()) : nullptr;
+			if (ValuesMapPtr)
+			{
+				ValuesMap = *ValuesMapPtr;
+			}
+		}
+		else
+		{
+			ValuesMap = Field.ToField()->GetMetaDataMap();
 		}
 	}
-	else
-	{
-		ValuesMap = Field.ToField()->GetMetaDataMap();
-	}
-	TArray<FString> MetaDataStrings;
+	
 	if (ValuesMap.Num())
 	{
 		for (auto& Pair : ValuesMap)
