Fix BP nativization dependency gathering to collect script objects referenced by an FProperty.

#rnx

diff --git a/Engine/Source/Developer/BlueprintCompilerCppBackend/Private/BlueprintCompilerCppBackendGatherDependencies.cpp b/Engine/Source/Developer/BlueprintCompilerCppBackend/Private/BlueprintCompilerCppBackendGatherDependencies.cpp
--- a/Engine/Source/Developer/BlueprintCompilerCppBackend/Private/BlueprintCompilerCppBackendGatherDependencies.cpp
+++ b/Engine/Source/Developer/BlueprintCompilerCppBackend/Private/BlueprintCompilerCppBackendGatherDependencies.cpp
@@ -21,6 +21,42 @@
 #include "K2Node_EnumLiteral.h"
 #include "KismetCompiler.h" // For LogK2Compiler
 
+class FGatherConvertedClassDependenciesArchive : public FReferenceCollectorArchive
+{
+public:
+	FGatherConvertedClassDependenciesArchive(UObject* InSerializingObject, FReferenceCollector& InCollector)
+		: FReferenceCollectorArchive(InSerializingObject, InCollector)
+	{
+		ArIsObjectReferenceCollector = true;
+		SetIsPersistent(InCollector.IsIgnoringTransient());
+		ArIgnoreArchetypeRef = InCollector.IsIgnoringArchetypeRef();
+	}
+
+protected:
+	virtual FArchive& operator<<(UObject*& Object) override
+	{
+		if (Object)
+		{
+			FReferenceCollector& CurrentCollector = GetCollector();
+			FProperty* OldCollectorSerializedProperty = CurrentCollector.GetSerializedProperty();
+			CurrentCollector.SetSerializedProperty(GetSerializedProperty());
+			CurrentCollector.AddReferencedObject(Object, GetSerializingObject(), GetSerializedProperty());
+			CurrentCollector.SetSerializedProperty(OldCollectorSerializedProperty);
+		}
+		return *this;
+	}
+
+	virtual FArchive& operator<<(FField*& Field) override
+	{
+		if (Field)
+		{
+			FReferenceCollector& CurrentCollector = GetCollector();
+			Field->AddReferencedObjects(CurrentCollector);
+		}
+		return *this;
+	}
+};
+
 struct FGatherConvertedClassDependenciesHelperBase : public FReferenceCollector
 {
 	TSet<UObject*> SerializedObjects;
@@ -28,6 +64,7 @@ struct FGatherConvertedClassDependenciesHelperBase : public FReferenceCollector
 
 	FGatherConvertedClassDependenciesHelperBase(FGatherConvertedClassDependencies& InDependencies)
 		: Dependencies(InDependencies)
+		, Archive(nullptr, *this)
 	{ }
 
 	virtual bool IsIgnoringArchetypeRef() const override { return false; }
@@ -36,7 +73,7 @@ struct FGatherConvertedClassDependenciesHelperBase : public FReferenceCollector
 	void FindReferences(UObject* Object)
 	{
 		{
-			FVerySlowReferenceCollectorArchiveScope CollectorScope(GetVerySlowReferenceCollectorArchive(), Object, nullptr);
+			FVerySlowReferenceCollectorArchiveScope CollectorScope(Archive, Object, nullptr);
 			const bool bOldFilterEditorOnly = CollectorScope.GetArchive().IsFilterEditorOnly();
 			CollectorScope.GetArchive().SetFilterEditorOnly(true);
 			if (UClass* AsBPGC = Cast<UBlueprintGeneratedClass>(Object))
@@ -45,6 +82,16 @@ struct FGatherConvertedClassDependenciesHelperBase : public FReferenceCollector
 			}
 			Object->Serialize(CollectorScope.GetArchive());
 			CollectorScope.GetArchive().SetFilterEditorOnly(bOldFilterEditorOnly);
+
+			// Class/struct fields (FFields) are no longer serialized, but we still need to find and collect property object references as part of dependency gathering.
+			// Note: Cannot rely on UStruct::PropertyObjectReferences (cached at link time) because that's GC-specific and excludes native references that we need here.
+			if (UStruct* ObjAsStruct = Cast<UStruct>(Object))
+			{
+				for (FField* ChildProperty = ObjAsStruct->ChildProperties; ChildProperty; ChildProperty = ChildProperty->Next)
+				{
+					ChildProperty->AddReferencedObjects(*this);
+				}
+			}
 		}
 	}
 
@@ -97,6 +144,9 @@ struct FGatherConvertedClassDependenciesHelperBase : public FReferenceCollector
 			Dependencies.ConvertedEnum.Add(InUDE);
 		}
 	}
+
+private:
+	FGatherConvertedClassDependenciesArchive Archive;
 };
 
 struct FFindAssetsToInclude : public FGatherConvertedClassDependenciesHelperBase
diff --git a/Engine/Source/Runtime/CoreUObject/Public/UObject/PropertyProxyArchive.h b/Engine/Source/Runtime/CoreUObject/Public/UObject/PropertyProxyArchive.h
--- a/Engine/Source/Runtime/CoreUObject/Public/UObject/PropertyProxyArchive.h
+++ b/Engine/Source/Runtime/CoreUObject/Public/UObject/PropertyProxyArchive.h
@@ -35,7 +35,7 @@ public:
 	}
 	virtual FArchive& operator<<(FField*& Value) override
 	{
-		if (!IsPersistent())
+		if (!IsPersistent() || IsObjectReferenceCollector())
 		{
 			// For reference collectors (like FArchiveReplaceFieldReferences): fully serialize the entire field to find all of its UObject references
 			InnerArchive << Value;
