Fix non-functional collision in a nativized Actor Blueprint with one or more components set up to use a custom collision profile to override the component's inherited default collision response.

diff --git a/Engine/Source/Developer/BlueprintCompilerCppBackend/Private/BlueprintCompilerCppBackendValueHelper.cpp b/Engine/Source/Developer/BlueprintCompilerCppBackend/Private/BlueprintCompilerCppBackendValueHelper.cpp
--- a/Engine/Source/Developer/BlueprintCompilerCppBackend/Private/BlueprintCompilerCppBackendValueHelper.cpp
+++ b/Engine/Source/Developer/BlueprintCompilerCppBackend/Private/BlueprintCompilerCppBackendValueHelper.cpp
@@ -31,6 +31,7 @@
 #include "Engine/InheritableComponentHandler.h"
 #include "Engine/DynamicBlueprintBinding.h"
 #include "Blueprint/BlueprintSupport.h"
+#include "Engine/CollisionProfile.h"
 
 void FEmitDefaultValueHelper::OuterGenerate(FEmitterLocalContext& Context
 	, const FProperty* Property
@@ -867,24 +868,33 @@ protected:
 			// Get the current collision profile names for each.
 			const FName ComponentCollisionProfileName = Component->BodyInstance.GetCollisionProfileName();
 			const FName ComponentArchetypeCollisionProfileName = ComponentArchetype->BodyInstance.GetCollisionProfileName();
+			const bool bIsArchetypeUsingCustomCollisionProfile = ComponentArchetypeCollisionProfileName == UCollisionProfile::CustomCollisionProfileName;
 
 			// Initialize a new struct instance that matches the archetype (represents the default struct value inherited by the component template).
 			FStructOnScope BodyInstanceToCompare(FBodyInstance::StaticStruct());
 			FBodyInstance::StaticStruct()->CopyScriptStruct(BodyInstanceToCompare.GetStructMemory(), &ComponentArchetype->BodyInstance);
 
-			if (ComponentCollisionProfileName != ComponentArchetypeCollisionProfileName)
+			// If the component template's collision profile setting differs from the default value (or is custom), set it using the API to load the modified collision profile.
+			bool bResetCollisionProfileAtRuntime = false;
+			if (ComponentCollisionProfileName != ComponentArchetypeCollisionProfileName || bIsArchetypeUsingCustomCollisionProfile)
 			{
-				// If the component template's collision profile setting differs from the default value, set it using the API to load the modified collision profile.
 				// This will initialize the struct's default value in the same manner as will occur at runtime, so we don't emit redundant initialization code to the ctor.
 				((FBodyInstance*)BodyInstanceToCompare.GetStructMemory())->SetCollisionProfileName(ComponentCollisionProfileName);
 
-				// Now emit the code to call SetCollisionProfileName() at runtime to initialize the collision profile within the instanced UPrimitiveComponent (see notes above).
-				Context.AddLine(FString::Printf(TEXT("%s->SetCollisionProfileName(FName(TEXT(\"%s\")));"), *VariableName, *ComponentCollisionProfileName.ToString().ReplaceCharWithEscapedChar()));
+				// We must also emit the same call at runtime to override the default collision profile within the instanced UPrimitiveComponent (see notes above).
+				bResetCollisionProfileAtRuntime = true;
 			}
 
 			// Emit the C++ code needed to initialize the remainder of the struct's value within the component that's now being instanced as a default subobject within the converted context.
 			const FString PathToMember = FString::Printf(TEXT("%s->BodyInstance"), *VariableName);
 			FEmitDefaultValueHelper::InnerGenerate(Context, BodyInstanceProperty, PathToMember, (const uint8*)&Component->BodyInstance, BodyInstanceToCompare.GetStructMemory(), FEmitDefaultValueHelper::EPropertyGenerationControlFlags::IncludeFirstConstructionLine);
+
+			// For a custom collision profile override, we need to have initialized the ResponseArray at runtime prior to calling SetCollisionProfileName(), so don't emit this call until after
+			// we've emitted the code above to initialize the remainder of the BodyInstance struct value. It's ok to emit this call last for any non-custom collision profile overrides as well.
+			if (bResetCollisionProfileAtRuntime)
+			{
+				Context.AddLine(FString::Printf(TEXT("%s->SetCollisionProfileName(FName(TEXT(\"%s\")));"), *VariableName, *ComponentCollisionProfileName.ToString().ReplaceCharWithEscapedChar()));
+			}
 		}
 		else
 		{
