Fixed blueprint nativization for SplineComponents/FInterpCurvePoints.

diff --git a/Engine/Source/Developer/BlueprintCompilerCppBackend/Private/BlueprintCompilerCppBackendUtils.cpp b/Engine/Source/Developer/BlueprintCompilerCppBackend/Private/BlueprintCompilerCppBackendUtils.cpp
--- a/Engine/Source/Developer/BlueprintCompilerCppBackend/Private/BlueprintCompilerCppBackendUtils.cpp
+++ b/Engine/Source/Developer/BlueprintCompilerCppBackend/Private/BlueprintCompilerCppBackendUtils.cpp
@@ -2406,6 +2406,12 @@ private:
 		MAP_BASE_STRUCTURE_ACCESS(TBaseStructure<FTransform>::Get());
 		MAP_BASE_STRUCTURE_ACCESS(TBaseStructure<FBox2D>::Get());
 		MAP_BASE_STRUCTURE_ACCESS(TBaseStructure<FFallbackStruct>::Get());
+		MAP_BASE_STRUCTURE_ACCESS(TBaseStructure<FInterpCurvePointFloat>::Get());
+		MAP_BASE_STRUCTURE_ACCESS(TBaseStructure<FInterpCurvePointVector2D>::Get());
+		MAP_BASE_STRUCTURE_ACCESS(TBaseStructure<FInterpCurvePointVector>::Get());
+		MAP_BASE_STRUCTURE_ACCESS(TBaseStructure<FInterpCurvePointQuat>::Get());
+		MAP_BASE_STRUCTURE_ACCESS(TBaseStructure<FInterpCurvePointTwoVectors>::Get());
+		MAP_BASE_STRUCTURE_ACCESS(TBaseStructure<FInterpCurvePointLinearColor>::Get());
 		MAP_BASE_STRUCTURE_ACCESS(TBaseStructure<FFloatRangeBound>::Get());
 		MAP_BASE_STRUCTURE_ACCESS(TBaseStructure<FFloatRange>::Get());
 		MAP_BASE_STRUCTURE_ACCESS(TBaseStructure<FInt32RangeBound>::Get());
diff --git a/Engine/Source/Developer/BlueprintCompilerCppBackend/Private/BlueprintCompilerCppBackendValueHelper.cpp b/Engine/Source/Developer/BlueprintCompilerCppBackend/Private/BlueprintCompilerCppBackendValueHelper.cpp
--- a/Engine/Source/Developer/BlueprintCompilerCppBackend/Private/BlueprintCompilerCppBackendValueHelper.cpp
+++ b/Engine/Source/Developer/BlueprintCompilerCppBackend/Private/BlueprintCompilerCppBackendValueHelper.cpp
@@ -604,6 +604,138 @@ bool FEmitDefaultValueHelper::SpecialStructureConstructor(const UStruct* Struct,
 		}
 		return true;
 	}
+	
+	if (TBaseStructure<FInterpCurvePointFloat>::Get() == Struct)
+	{
+		if (OutResult)
+		{
+			const FInterpCurvePointFloat* CurvePointFloat = reinterpret_cast<const FInterpCurvePointFloat*>(ValuePtr);
+			*OutResult = FString::Printf(TEXT("CreateFInterpCurvePointFloat(%s, %s, %s, %s, TEnumAsByte<EInterpCurveMode>(%d))")
+				, *FEmitHelper::FloatToString(CurvePointFloat->InVal)
+				, *FEmitHelper::FloatToString(CurvePointFloat->OutVal)
+				, *FEmitHelper::FloatToString(CurvePointFloat->ArriveTangent)
+				, *FEmitHelper::FloatToString(CurvePointFloat->LeaveTangent)
+				, CurvePointFloat->InterpMode.GetValue());
+		}
+		return true;
+	}
+
+	if (TBaseStructure<FInterpCurvePointVector2D>::Get() == Struct)
+	{
+		if (OutResult)
+		{
+			const FInterpCurvePointVector2D* CurvePointVector2D = reinterpret_cast<const FInterpCurvePointVector2D*>(ValuePtr);
+			*OutResult = FString::Printf(TEXT("CreateFInterpCurvePointVector2D(%s, FVector2D(%s, %s), FVector2D(%s, %s), FVector2D(%s, %s), TEnumAsByte<EInterpCurveMode>(%d))")
+				, *FEmitHelper::FloatToString(CurvePointVector2D->InVal)
+				, *FEmitHelper::FloatToString(CurvePointVector2D->OutVal.X)
+				, *FEmitHelper::FloatToString(CurvePointVector2D->OutVal.Y)
+				, *FEmitHelper::FloatToString(CurvePointVector2D->ArriveTangent.X)
+				, *FEmitHelper::FloatToString(CurvePointVector2D->ArriveTangent.Y)
+				, *FEmitHelper::FloatToString(CurvePointVector2D->LeaveTangent.X)
+				, *FEmitHelper::FloatToString(CurvePointVector2D->LeaveTangent.Y)
+				, CurvePointVector2D->InterpMode.GetValue());
+		}
+		return true;
+	}
+
+	if (TBaseStructure<FInterpCurvePointVector>::Get() == Struct)
+	{
+		if (OutResult)
+		{
+			const FInterpCurvePointVector* CurvePointVector = reinterpret_cast<const FInterpCurvePointVector*>(ValuePtr);
+			*OutResult = FString::Printf(TEXT("CreateFInterpCurvePointVector(%s, FVector(%s, %s, %s), FVector(%s, %s, %s), FVector(%s, %s, %s), TEnumAsByte<EInterpCurveMode>(%d))")
+				, *FEmitHelper::FloatToString(CurvePointVector->InVal)
+				, *FEmitHelper::FloatToString(CurvePointVector->OutVal.X)
+				, *FEmitHelper::FloatToString(CurvePointVector->OutVal.Y)
+				, *FEmitHelper::FloatToString(CurvePointVector->OutVal.Z)
+				, *FEmitHelper::FloatToString(CurvePointVector->ArriveTangent.X)
+				, *FEmitHelper::FloatToString(CurvePointVector->ArriveTangent.Y)
+				, *FEmitHelper::FloatToString(CurvePointVector->ArriveTangent.Z)
+				, *FEmitHelper::FloatToString(CurvePointVector->LeaveTangent.X)
+				, *FEmitHelper::FloatToString(CurvePointVector->LeaveTangent.Y)
+				, *FEmitHelper::FloatToString(CurvePointVector->LeaveTangent.Z)
+				, CurvePointVector->InterpMode.GetValue());
+		}
+		return true;
+	}
+
+	if (TBaseStructure<FInterpCurvePointQuat>::Get() == Struct)
+	{
+		if (OutResult)
+		{
+			const FInterpCurvePointQuat* CurvePointQuat = reinterpret_cast<const FInterpCurvePointQuat*>(ValuePtr);
+			*OutResult = FString::Printf(TEXT("CreateFInterpCurvePointQuat(%s, FQuat(%s, %s, %s, %s), FQuat(%s, %s, %s, %s), FQuat(%s, %s, %s, %s), TEnumAsByte<EInterpCurveMode>(%d))")
+				, *FEmitHelper::FloatToString(CurvePointQuat->InVal)
+				, *FEmitHelper::FloatToString(CurvePointQuat->OutVal.X)
+				, *FEmitHelper::FloatToString(CurvePointQuat->OutVal.Y)
+				, *FEmitHelper::FloatToString(CurvePointQuat->OutVal.Z)
+				, *FEmitHelper::FloatToString(CurvePointQuat->OutVal.W)
+				, *FEmitHelper::FloatToString(CurvePointQuat->ArriveTangent.X)
+				, *FEmitHelper::FloatToString(CurvePointQuat->ArriveTangent.Y)
+				, *FEmitHelper::FloatToString(CurvePointQuat->ArriveTangent.Z)
+				, *FEmitHelper::FloatToString(CurvePointQuat->ArriveTangent.W)
+				, *FEmitHelper::FloatToString(CurvePointQuat->LeaveTangent.X)
+				, *FEmitHelper::FloatToString(CurvePointQuat->LeaveTangent.Y)
+				, *FEmitHelper::FloatToString(CurvePointQuat->LeaveTangent.Z)
+				, *FEmitHelper::FloatToString(CurvePointQuat->LeaveTangent.W)
+				, CurvePointQuat->InterpMode.GetValue());
+		}
+		return true;
+	}
+
+	if (TBaseStructure<FInterpCurvePointTwoVectors>::Get() == Struct)
+	{
+		if (OutResult)
+		{
+			const FInterpCurvePointTwoVectors* CurvePointTwoVectors = reinterpret_cast<const FInterpCurvePointTwoVectors*>(ValuePtr);
+			*OutResult = FString::Printf(TEXT("CreateFInterpCurvePointTwoVectors(%s, FVector(%s, %s, %s), FVector(%s, %s, %s), FVector(%s, %s, %s), FVector(%s, %s, %s), FVector(%s, %s, %s), FVector(%s, %s, %s), TEnumAsByte<EInterpCurveMode>(%d))")
+				, *FEmitHelper::FloatToString(CurvePointTwoVectors->InVal)
+				, *FEmitHelper::FloatToString(CurvePointTwoVectors->OutVal.v1.X)
+				, *FEmitHelper::FloatToString(CurvePointTwoVectors->OutVal.v1.Y)
+				, *FEmitHelper::FloatToString(CurvePointTwoVectors->OutVal.v1.Z)
+				, *FEmitHelper::FloatToString(CurvePointTwoVectors->OutVal.v2.X)
+				, *FEmitHelper::FloatToString(CurvePointTwoVectors->OutVal.v2.Y)
+				, *FEmitHelper::FloatToString(CurvePointTwoVectors->OutVal.v2.Z)
+				, *FEmitHelper::FloatToString(CurvePointTwoVectors->ArriveTangent.v1.X)
+				, *FEmitHelper::FloatToString(CurvePointTwoVectors->ArriveTangent.v1.Y)
+				, *FEmitHelper::FloatToString(CurvePointTwoVectors->ArriveTangent.v1.Z)
+				, *FEmitHelper::FloatToString(CurvePointTwoVectors->ArriveTangent.v2.X)
+				, *FEmitHelper::FloatToString(CurvePointTwoVectors->ArriveTangent.v2.Y)
+				, *FEmitHelper::FloatToString(CurvePointTwoVectors->ArriveTangent.v2.Z)
+				, *FEmitHelper::FloatToString(CurvePointTwoVectors->LeaveTangent.v1.X)
+				, *FEmitHelper::FloatToString(CurvePointTwoVectors->LeaveTangent.v1.Y)
+				, *FEmitHelper::FloatToString(CurvePointTwoVectors->LeaveTangent.v1.Z)
+				, *FEmitHelper::FloatToString(CurvePointTwoVectors->LeaveTangent.v2.X)
+				, *FEmitHelper::FloatToString(CurvePointTwoVectors->LeaveTangent.v2.Y)
+				, *FEmitHelper::FloatToString(CurvePointTwoVectors->LeaveTangent.v2.Z)
+				, CurvePointTwoVectors->InterpMode.GetValue());
+		}
+		return true;
+	}
+
+	if (TBaseStructure<FInterpCurvePointLinearColor>::Get() == Struct)
+	{
+		if (OutResult)
+		{
+			const FInterpCurvePointLinearColor* CurvePointLinearColor = reinterpret_cast<const FInterpCurvePointLinearColor*>(ValuePtr);
+			*OutResult = FString::Printf(TEXT("CreateFInterpCurvePointLinearColor(%s, FLinearColor(%s, %s, %s, %s), FLinearColor(%s, %s, %s, %s), FLinearColor(%s, %s, %s, %s), TEnumAsByte<EInterpCurveMode>(%d))")
+				, *FEmitHelper::FloatToString(CurvePointLinearColor->InVal)
+				, *FEmitHelper::FloatToString(CurvePointLinearColor->OutVal.R)
+				, *FEmitHelper::FloatToString(CurvePointLinearColor->OutVal.G)
+				, *FEmitHelper::FloatToString(CurvePointLinearColor->OutVal.B)
+				, *FEmitHelper::FloatToString(CurvePointLinearColor->OutVal.A)
+				, *FEmitHelper::FloatToString(CurvePointLinearColor->ArriveTangent.R)
+				, *FEmitHelper::FloatToString(CurvePointLinearColor->ArriveTangent.G)
+				, *FEmitHelper::FloatToString(CurvePointLinearColor->ArriveTangent.B)
+				, *FEmitHelper::FloatToString(CurvePointLinearColor->ArriveTangent.A)
+				, *FEmitHelper::FloatToString(CurvePointLinearColor->LeaveTangent.R)
+				, *FEmitHelper::FloatToString(CurvePointLinearColor->LeaveTangent.G)
+				, *FEmitHelper::FloatToString(CurvePointLinearColor->LeaveTangent.B)
+				, *FEmitHelper::FloatToString(CurvePointLinearColor->LeaveTangent.A)
+				, CurvePointLinearColor->InterpMode.GetValue());
+		}
+		return true;
+	}
 
 	if (TBaseStructure<FFloatRangeBound>::Get() == Struct)
 	{
diff --git a/Engine/Source/Runtime/Core/Public/Math/InterpCurvePoint.h b/Engine/Source/Runtime/Core/Public/Math/InterpCurvePoint.h
--- a/Engine/Source/Runtime/Core/Public/Math/InterpCurvePoint.h
+++ b/Engine/Source/Runtime/Core/Public/Math/InterpCurvePoint.h
@@ -93,6 +93,13 @@ public:
 	 */
 	FInterpCurvePoint( const float In, const T &Out, const T &InArriveTangent, const T &InLeaveTangent, const EInterpCurveMode InInterpMode );
 
+	/**
+	 * Constructor which initializes all components to zero.
+	 *
+	 * @param EForceInit Force init enum
+	 */
+	explicit FORCEINLINE FInterpCurvePoint(EForceInit);
+
 public:
 
 	/** @return true if the key value is using a curve interp mode, otherwise false */
@@ -162,6 +169,15 @@ FORCEINLINE FInterpCurvePoint<T>::FInterpCurvePoint( const float In, const T &Ou
 	, InterpMode(InInterpMode)
 { }
 
+template< class T >
+FORCEINLINE FInterpCurvePoint<T>::FInterpCurvePoint(EForceInit)
+{
+	InVal = 0.0f;
+	FMemory::Memset(&OutVal, 0, sizeof(T));
+	FMemory::Memset(&ArriveTangent, 0, sizeof(T));
+	FMemory::Memset(&LeaveTangent, 0, sizeof(T));
+	InterpMode = CIM_Linear;
+}
 
 template< class T > 
 FORCEINLINE bool FInterpCurvePoint<T>::IsCurveKey() const
diff --git a/Engine/Source/Runtime/CoreUObject/Private/UObject/Class.cpp b/Engine/Source/Runtime/CoreUObject/Private/UObject/Class.cpp
--- a/Engine/Source/Runtime/CoreUObject/Private/UObject/Class.cpp
+++ b/Engine/Source/Runtime/CoreUObject/Private/UObject/Class.cpp
@@ -53,7 +53,7 @@
 #include "UObject/PropertyProxyArchive.h"
 #include "UObject/FieldPath.h"
 #include "HAL/ThreadSafeCounter.h"
-
+#include "Math/InterpCurvePoint.h"
 
 // WARNING: This should always be the last include in any file that needs it (except .generated.h)
 #include "UObject/UndefineUPropertyMacros.h"
@@ -5913,6 +5913,42 @@ UScriptStruct* TBaseStructure<FFallbackStruct>::Get()
 	return ScriptStruct;
 }
 
+UScriptStruct* TBaseStructure<FInterpCurvePointFloat>::Get()
+{
+	static UScriptStruct* ScriptStruct = StaticGetBaseStructureInternal(TEXT("InterpCurvePointFloat"));
+	return ScriptStruct;
+}
+
+UScriptStruct* TBaseStructure<FInterpCurvePointVector2D>::Get()
+{
+	static UScriptStruct* ScriptStruct = StaticGetBaseStructureInternal(TEXT("InterpCurvePointVector2D"));
+	return ScriptStruct;
+}
+
+UScriptStruct* TBaseStructure<FInterpCurvePointVector>::Get()
+{
+	static UScriptStruct* ScriptStruct = StaticGetBaseStructureInternal(TEXT("InterpCurvePointVector"));
+	return ScriptStruct;
+}
+
+UScriptStruct* TBaseStructure<FInterpCurvePointQuat>::Get()
+{
+	static UScriptStruct* ScriptStruct = StaticGetBaseStructureInternal(TEXT("InterpCurvePointQuat"));
+	return ScriptStruct;
+}
+
+UScriptStruct* TBaseStructure<FInterpCurvePointTwoVectors>::Get()
+{
+	static UScriptStruct* ScriptStruct = StaticGetBaseStructureInternal(TEXT("InterpCurvePointTwoVectors"));
+	return ScriptStruct;
+}
+
+UScriptStruct* TBaseStructure<FInterpCurvePointLinearColor>::Get()
+{
+	static UScriptStruct* ScriptStruct = StaticGetBaseStructureInternal(TEXT("InterpCurvePointLinearColor"));
+	return ScriptStruct;
+}
+
 UScriptStruct* TBaseStructure<FFloatRangeBound>::Get()
 {
 	static UScriptStruct* ScriptStruct = StaticGetBaseStructureInternal(TEXT("FloatRangeBound"));
diff --git a/Engine/Source/Runtime/CoreUObject/Private/UObject/Property.cpp b/Engine/Source/Runtime/CoreUObject/Private/UObject/Property.cpp
--- a/Engine/Source/Runtime/CoreUObject/Private/UObject/Property.cpp
+++ b/Engine/Source/Runtime/CoreUObject/Private/UObject/Property.cpp
@@ -20,6 +20,7 @@
 #include "UObject/CoreRedirects.h"
 #include "UObject/SoftObjectPath.h"
 #include "Math/Box2D.h"
+#include "Math/InterpCurvePoint.h"
 #include "UObject/ReleaseObjectVersion.h"
 
 // WARNING: This should always be the last include in any file that needs it (except .generated.h)
@@ -237,6 +238,66 @@ struct TStructOpsTypeTraits<FTwoVectors> : public TStructOpsTypeTraitsBase2<FTwo
 };
 IMPLEMENT_STRUCT(TwoVectors);
 
+template<>
+struct TStructOpsTypeTraits<FInterpCurvePointFloat> : public TStructOpsTypeTraitsBase2<FInterpCurvePointFloat>
+{
+	enum
+	{
+		WithNoInitConstructor = true,
+	};
+};
+IMPLEMENT_STRUCT(InterpCurvePointFloat);
+
+template<>
+struct TStructOpsTypeTraits<FInterpCurvePointVector2D> : public TStructOpsTypeTraitsBase2<FInterpCurvePointVector2D>
+{
+	enum
+	{
+		WithNoInitConstructor = true,
+	};
+};
+IMPLEMENT_STRUCT(InterpCurvePointVector2D);
+
+template<>
+struct TStructOpsTypeTraits<FInterpCurvePointVector> : public TStructOpsTypeTraitsBase2<FInterpCurvePointVector>
+{
+	enum
+	{
+		WithNoInitConstructor = true,
+	};
+};
+IMPLEMENT_STRUCT(InterpCurvePointVector);
+
+template<>
+struct TStructOpsTypeTraits<FInterpCurvePointQuat> : public TStructOpsTypeTraitsBase2<FInterpCurvePointQuat>
+{
+	enum
+	{
+		WithNoInitConstructor = true,
+	};
+};
+IMPLEMENT_STRUCT(InterpCurvePointQuat);
+
+template<>
+struct TStructOpsTypeTraits<FInterpCurvePointTwoVectors> : public TStructOpsTypeTraitsBase2<FInterpCurvePointTwoVectors>
+{
+	enum
+	{
+		WithNoInitConstructor = true,
+	};
+};
+IMPLEMENT_STRUCT(InterpCurvePointTwoVectors);
+
+template<>
+struct TStructOpsTypeTraits<FInterpCurvePointLinearColor> : public TStructOpsTypeTraitsBase2<FInterpCurvePointLinearColor>
+{
+	enum
+	{
+		WithNoInitConstructor = true,
+	};
+};
+IMPLEMENT_STRUCT(InterpCurvePointLinearColor);
+
 template<>
 struct TStructOpsTypeTraits<FGuid> : public TStructOpsTypeTraitsBase2<FGuid>
 {
diff --git a/Engine/Source/Runtime/CoreUObject/Public/UObject/Class.h b/Engine/Source/Runtime/CoreUObject/Public/UObject/Class.h
--- a/Engine/Source/Runtime/CoreUObject/Public/UObject/Class.h
+++ b/Engine/Source/Runtime/CoreUObject/Public/UObject/Class.h
@@ -10,6 +10,7 @@
 
 #include "Concepts/GetTypeHashable.h"
 #include "Math/RandomStream.h"
+#include "Math/InterpCurvePoint.h"
 #include "Misc/EnumClassFlags.h"
 #include "Misc/FallbackStruct.h"
 #include "Misc/Guid.h"
@@ -3731,6 +3732,36 @@ template<> struct TBaseStructure<FFallbackStruct>
 	COREUOBJECT_API static UScriptStruct* Get();
 };
 
+template<> struct TBaseStructure<FInterpCurvePointFloat>
+{
+	COREUOBJECT_API static UScriptStruct* Get();
+};
+
+template<> struct TBaseStructure<FInterpCurvePointVector2D>
+{
+	COREUOBJECT_API static UScriptStruct* Get();
+};
+
+template<> struct TBaseStructure<FInterpCurvePointVector>
+{
+	COREUOBJECT_API static UScriptStruct* Get();
+};
+
+template<> struct TBaseStructure<FInterpCurvePointQuat>
+{
+	COREUOBJECT_API static UScriptStruct* Get();
+};
+
+template<> struct TBaseStructure<FInterpCurvePointTwoVectors>
+{
+	COREUOBJECT_API static UScriptStruct* Get();
+};
+
+template<> struct TBaseStructure<FInterpCurvePointLinearColor>
+{
+	COREUOBJECT_API static UScriptStruct* Get();
+};
+
 template<> struct TBaseStructure<FFloatRangeBound>
 {
 	COREUOBJECT_API static UScriptStruct* Get();
diff --git a/Engine/Source/Runtime/Engine/Public/GeneratedCodeHelpers.h b/Engine/Source/Runtime/Engine/Public/GeneratedCodeHelpers.h
--- a/Engine/Source/Runtime/Engine/Public/GeneratedCodeHelpers.h
+++ b/Engine/Source/Runtime/Engine/Public/GeneratedCodeHelpers.h
@@ -28,8 +28,9 @@
 //For DOREPLIFETIME macros
 #include "Net/UnrealNetwork.h"
 
-//For Box2D
 #include "Runtime/Core/Public/Math/Box2D.h"
+#include "Runtime/Core/Public/Math/InterpCurvePoint.h"
+#include "Runtime/Core/Public/Containers/EnumAsByte.h"
 
 inline FBox2D CreateFBox2D(FVector2D InMin, FVector2D InMax, bool InIsValid)
 {
@@ -40,6 +41,74 @@ inline FBox2D CreateFBox2D(FVector2D InMin, FVector2D InMax, bool InIsValid)
 	return Result;
 }
 
+inline FInterpCurvePointFloat CreateFInterpCurvePointFloat(float InVal, float OutVal, float ArriveTangent, float LeaveTangent, const TEnumAsByte<EInterpCurveMode>& InterpMode)
+{
+	FInterpCurvePointFloat Result;
+	Result.InVal = InVal;
+	Result.OutVal = OutVal;
+	Result.ArriveTangent = ArriveTangent;
+	Result.LeaveTangent = LeaveTangent;
+	Result.InterpMode = InterpMode;
+	return Result;
+}
+
+inline FInterpCurvePointVector2D CreateFInterpCurvePointVector2D(float InVal, const FVector2D& OutVal, const FVector2D& ArriveTangent, const FVector2D& LeaveTangent, const TEnumAsByte<EInterpCurveMode>& InterpMode)
+{
+	FInterpCurvePointVector2D Result;
+	Result.InVal = InVal;
+	Result.OutVal = OutVal;
+	Result.ArriveTangent = ArriveTangent;
+	Result.LeaveTangent = LeaveTangent;
+	Result.InterpMode = InterpMode;
+	return Result;
+}
+
+inline FInterpCurvePointVector CreateFInterpCurvePointVector(float InVal, const FVector& OutVal, const FVector& ArriveTangent, const FVector& LeaveTangent, const TEnumAsByte<EInterpCurveMode>& InterpMode)
+{
+	FInterpCurvePointVector Result;
+	Result.InVal = InVal;
+	Result.OutVal = OutVal;
+	Result.ArriveTangent = ArriveTangent;
+	Result.LeaveTangent = LeaveTangent;
+	Result.InterpMode = InterpMode;
+	return Result;
+}
+
+inline FInterpCurvePointQuat CreateFInterpCurvePointQuat(float InVal, const FQuat& OutVal, const FQuat& ArriveTangent, const FQuat& LeaveTangent, const TEnumAsByte<EInterpCurveMode>& InterpMode)
+{
+	FInterpCurvePointQuat Result;
+	Result.InVal = InVal;
+	Result.OutVal = OutVal;
+	Result.ArriveTangent = ArriveTangent;
+	Result.LeaveTangent = LeaveTangent;
+	Result.InterpMode = InterpMode;
+	return Result;
+}
+
+inline FInterpCurvePointTwoVectors CreateFInterpCurvePointTwoVectors(float InVal, const FVector& OutVal1, const FVector& OutVal2, const FVector& ArriveTangent1,
+									const FVector& ArriveTangent2, const FVector& LeaveTangent1, const FVector& LeaveTangent2,
+									const EInterpCurveMode& InterpMode)
+{
+	FInterpCurvePointTwoVectors Result;
+	Result.InVal = InVal;
+	Result.OutVal = FTwoVectors(OutVal1, OutVal2);
+	Result.ArriveTangent = FTwoVectors(ArriveTangent1, ArriveTangent2);
+	Result.LeaveTangent = FTwoVectors(LeaveTangent1, LeaveTangent2);
+	Result.InterpMode = InterpMode;
+	return Result;
+}
+
+inline FInterpCurvePointLinearColor CreateFInterpCurvePointLinearColor(float InVal, const FLinearColor& OutVal, const FLinearColor& ArriveTangent, const FLinearColor& LeaveTangent, const EInterpCurveMode& InterpMode)
+{
+	FInterpCurvePointLinearColor Result;
+	Result.InVal = InVal;
+	Result.OutVal = OutVal;
+	Result.ArriveTangent = ArriveTangent;
+	Result.LeaveTangent = LeaveTangent;
+	Result.InterpMode = InterpMode;
+	return Result;
+}
+
 template<class NativeType>
 inline NativeType* NoNativeCast(UClass* NoNativeClass, UObject* Object)
 {
diff --git a/Engine/Source/Runtime/PropertyPath/Public/PropertyTypeCompatibility.h b/Engine/Source/Runtime/PropertyPath/Public/PropertyTypeCompatibility.h
--- a/Engine/Source/Runtime/PropertyPath/Public/PropertyTypeCompatibility.h
+++ b/Engine/Source/Runtime/PropertyPath/Public/PropertyTypeCompatibility.h
@@ -8,7 +8,7 @@
 #include "UObject/UnrealType.h"
 #include "UObject/EnumProperty.h"
 #include "UObject/TextProperty.h"
-
+#include "Math/InterpCurvePoint.h"
 
 /** Generic (struct) implementation */
 template<typename T>
@@ -229,6 +229,42 @@ inline bool IsConcreteTypeCompatibleWithReflectedType_Impl<FGuid>(FProperty* Pro
 	return IsConcreteTypeCompatibleWithReflectedType_BuiltInStruct<FGuid>(Property);
 }
 
+template<>
+inline bool IsConcreteTypeCompatibleWithReflectedType_Impl<FInterpCurvePointFloat>(FProperty* Property)
+{
+	return IsConcreteTypeCompatibleWithReflectedType_BuiltInStruct<FInterpCurvePointFloat>(Property);
+}
+
+template<>
+inline bool IsConcreteTypeCompatibleWithReflectedType_Impl<FInterpCurvePointVector>(FProperty* Property)
+{
+	return IsConcreteTypeCompatibleWithReflectedType_BuiltInStruct<FInterpCurvePointVector>(Property);
+}
+
+template<>
+inline bool IsConcreteTypeCompatibleWithReflectedType_Impl<FInterpCurvePointVector2D>(FProperty* Property)
+{
+	return IsConcreteTypeCompatibleWithReflectedType_BuiltInStruct<FInterpCurvePointVector2D>(Property);
+}
+
+template<>
+inline bool IsConcreteTypeCompatibleWithReflectedType_Impl<FInterpCurvePointQuat>(FProperty* Property)
+{
+	return IsConcreteTypeCompatibleWithReflectedType_BuiltInStruct<FInterpCurvePointQuat>(Property);
+}
+
+template<>
+inline bool IsConcreteTypeCompatibleWithReflectedType_Impl<FInterpCurvePointTwoVectors>(FProperty* Property)
+{
+	return IsConcreteTypeCompatibleWithReflectedType_BuiltInStruct<FInterpCurvePointTwoVectors>(Property);
+}
+
+template<>
+inline bool IsConcreteTypeCompatibleWithReflectedType_Impl<FInterpCurvePointLinearColor>(FProperty* Property)
+{
+	return IsConcreteTypeCompatibleWithReflectedType_BuiltInStruct<FInterpCurvePointLinearColor>(Property);
+}
+
 template<>
 inline bool IsConcreteTypeCompatibleWithReflectedType_Impl<FFloatRangeBound>(FProperty* Property)
 {
